Move OnlineTests to public folder.

diff --git a/Engine/Source/Programs/Online/OnlineTestsCore/Source/OnlineTestsCore.Build.cs b/Engine/Source/Programs/Online/OnlineTestsCore/Source/OnlineTestsCore.Build.cs
new file mode 100644
--- /dev/null
+++ b/Engine/Source/Programs/Online/OnlineTestsCore/Source/OnlineTestsCore.Build.cs
@@ -0,0 +1,61 @@
+// Copyright Epic Games, Inc. All Rights Reserved.
+
+using UnrealBuildTool;
+using System.IO;
+using System;
+using System.Collections.Generic;
+using System.Linq;
+
+public class OnlineTestsCore : ModuleRules
+{
+	public virtual bool bRequireApplicationTick { get { return false; } }
+	public virtual bool bRequirePlatformInit { get { return false; } }
+	public virtual string PlatformFileName { get { return ""; } }
+	public virtual bool bUseExternAuth { get { return false; } }
+
+	public OnlineTestsCore(ReadOnlyTargetRules Target) : base(Target)
+	{
+		bTreatAsEngineModule = false;
+		CppStandard = CppStandardVersion.EngineDefault;
+		IncludeOrderVersion = EngineIncludeOrderVersion.Latest;
+
+		PublicDependencyModuleNames.AddRange(
+			new string[] {
+				"ApplicationCore",
+				"Core",
+				"CoreUObject",
+				"Projects",
+				"EngineSettings",
+				"OnlineSubsystem",
+				"OnlineServicesInterface",
+				"OnlineServicesCommon",
+				"OnlineServicesEOS",
+				"OnlineServicesNull",
+				"OnlineServicesOSSAdapter",
+				"EOSSDK",
+				"EOSShared",
+				"SSL",
+				"Json",
+				"JsonUtilities"
+			}
+		);
+
+		PublicDependencyModuleNames.AddRange(
+		new string[]
+		{
+			"OnlineSubsystem"
+		});
+
+		PublicDefinitions.Add(String.Format("ONLINETESTS_REQUIREAPPLICATIONTICK={0}", bRequireApplicationTick ? 1 : 0));
+		PublicDefinitions.Add(String.Format("ONLINETESTS_REQUIREPLATFORMINIT={0}", bRequirePlatformInit ? 1 : 0));
+		PublicDefinitions.Add(String.Format("ONLINETESTS_PLATFORMFILENAME={0}", PlatformFileName));
+
+		// Disable external auth if target doesn't define it.
+		if (!Target.GlobalDefinitions.Contains("ONLINETESTS_USEEXTERNAUTH=1"))
+		{
+			PublicDefinitions.Add(String.Format("ONLINETESTS_USEEXTERNAUTH=0"));
+		}
+	}
+}
+
+
diff --git a/Engine/Source/Programs/Online/OnlineTestsCore/Source/Private/Tests/EOS/TestErrorsEOS.cpp b/Engine/Source/Programs/Online/OnlineTestsCore/Source/Private/Tests/EOS/TestErrorsEOS.cpp
new file mode 100644
--- /dev/null
+++ b/Engine/Source/Programs/Online/OnlineTestsCore/Source/Private/Tests/EOS/TestErrorsEOS.cpp
@@ -0,0 +1,58 @@
+// Copyright Epic Games, Inc. All Rights Reserved.
+#include "CoreMinimal.h"
+#include <catch2/catch_test_macros.hpp>
+
+#include "Online/OnlineError.h"
+#include "Online/OnlineErrorDefinitions.h"
+#include "Online/OnlineErrorEOSGS.h"
+
+#define LOCTEXT_NAMESPACE "OnlineErrorsTest"
+#define EOS_ERROR_TAG "[.EOSErrors]"
+#define EOS_ERROR_TEST_CASE(x, ...) TEST_CASE(x, EOS_ERROR_TAG __VA_ARGS__)
+
+// Dummy error category for testing mapping extensions.
+namespace UE::Online::Errors
+{
+	UE_ONLINE_ERROR_CATEGORY(ExtensionTest, Engine, 0x999, "ExtensionTest")
+	UE_ONLINE_ERROR(ExtensionTest, WrongAccount, 1, TEXT("WrongAccount"), LOCTEXT("WrongAccount", ""))
+/* UE::Online::Errors */ }
+
+UE::Online::FOnlineError MapExtensionEOSError(UE::Online::FOnlineError&& Error, EOS_EResult Result)
+{
+	switch (Result)
+	{
+	case EOS_EResult::EOS_Auth_WrongAccount:	return UE::Online::Errors::ExtensionTest::WrongAccount(MoveTemp(Error));
+	default:									return UE::Online::Errors::MapCommonEOSError(MoveTemp(Error), Result);
+	}
+}
+
+EOS_ERROR_TEST_CASE("Basic EOS error compilation test")
+{
+	using namespace UE::Online;
+	using namespace UE::Online::Errors;
+
+	FOnlineError TestError_Success = Errors::FromEOSResult(EOS_EResult::EOS_Success);
+	FOnlineError TestError_InvalidParams = Errors::FromEOSResult(EOS_EResult::EOS_InvalidParameters);
+	FOnlineError TestError_Canceled = Errors::FromEOSResult(EOS_EResult::EOS_Canceled);
+	FOnlineError TestError_Big = Errors::FromEOSResult(EOS_EResult::EOS_PlayerDataStorage_DataLengthInvalid);
+	FOnlineError TestError_ExtensionKnown = Errors::FromEOSResult(EOS_EResult::EOS_Auth_WrongAccount, &MapExtensionEOSError);
+	FOnlineError TestError_ExtensionCommon = Errors::FromEOSResult(EOS_EResult::EOS_Success, &MapExtensionEOSError);
+
+	CHECK(TestError_Success == Errors::Success());
+	CHECK(TestError_Success == EOS_EResult::EOS_Success);
+	CHECK(TestError_Success != EOS_EResult::EOS_InvalidParameters);
+	CHECK(TestError_InvalidParams == Errors::InvalidParams());
+	CHECK(TestError_InvalidParams == EOS_EResult::EOS_InvalidParameters);
+	CHECK(TestError_InvalidParams != EOS_EResult::EOS_Canceled);
+	CHECK(TestError_Canceled == Errors::Cancelled());
+	CHECK(TestError_Canceled == EOS_EResult::EOS_Canceled);
+	CHECK(TestError_Canceled != EOS_EResult::EOS_Success);
+	CHECK(TestError_Big == Errors::Unknown());
+	CHECK(TestError_Big == EOS_EResult::EOS_PlayerDataStorage_DataLengthInvalid);
+	CHECK(TestError_ExtensionKnown == Errors::ExtensionTest::WrongAccount());
+	CHECK(TestError_ExtensionKnown == EOS_EResult::EOS_Auth_WrongAccount);
+	CHECK(TestError_ExtensionCommon == Errors::Success());
+	CHECK(TestError_ExtensionCommon == EOS_EResult::EOS_Success);
+}
+
+#undef LOCTEXT_NAMESPACE
\ No newline at end of file
diff --git a/Engine/Source/Programs/Online/OnlineTestsCore/Source/Private/Tests/Leaderboards/LeaderboardsReadEntriesAroundRankTests.cpp b/Engine/Source/Programs/Online/OnlineTestsCore/Source/Private/Tests/Leaderboards/LeaderboardsReadEntriesAroundRankTests.cpp
new file mode 100644
--- /dev/null
+++ b/Engine/Source/Programs/Online/OnlineTestsCore/Source/Private/Tests/Leaderboards/LeaderboardsReadEntriesAroundRankTests.cpp
@@ -0,0 +1,86 @@
+// Copyright Epic Games, Inc. All Rights Reserved.
+
+#include "ReadEntriesAroundRankFixture.h"
+#include "LeaderboardsFixture.h"
+
+#define LEADERBOARDS_TAGS "[Leaderboards]"
+#define LEADERBOARDS_TEST_CASE(x, ...) ONLINE_TEST_CASE(x, LEADERBOARDS_TAGS __VA_ARGS__)
+
+
+// ReadEntriesAroundRank tests
+
+LEADERBOARDS_TEST_CASE("Verify ReadEntriesAroundRank succeeds and returns entries for the given board", "[MultiAccount]")
+{
+	TArray<UE::Online::FLeaderboardEntry> LeaderboardEntries;
+	FAccountId AccountIdA, AccountIdB;
+	TArray< FAccountId*> AccountIds = { &AccountIdA, &AccountIdB };
+	uint32 NumUsersToLogin = 2;
+
+	LeaderboardsFixture_NUsers(std::move(GetLoginPipeline(AccountIdA, AccountIdB)), AccountIds, NumUsersToLogin)
+		.EmplaceAsyncLambda([&](FAsyncLambdaResult Promise, SubsystemType Services) { ReadEntriesAroundRank_2UsersFixture(Promise, Services, "Stat_Use_Set", AccountIdA, 0, 2, LeaderboardEntries); })
+		.EmplaceLambda([&](SubsystemType Services) { CheckEntries_2Users_Stat_Use_Set(LeaderboardEntries, AccountIdA, AccountIdB); })
+		.EmplaceAsyncLambda([&](FAsyncLambdaResult Promise, SubsystemType Services) { ReadEntriesAroundRank_2UsersFixture(Promise, Services, "Stat_Use_Smallest", AccountIdA, 0, 2, LeaderboardEntries); })
+		.EmplaceLambda([&](SubsystemType Services) { CheckEntries_2Users_Stat_Use_Smallest(LeaderboardEntries, AccountIdA, AccountIdB); });
+
+	RunToCompletion();
+}
+
+LEADERBOARDS_TEST_CASE("Verify ReadEntriesAroundRank returns a fail message if the local user is not logged in", "[MultiAccount]")
+{
+	TArray<UE::Online::FLeaderboardEntry> LeaderboardEntries;
+	FAccountId AccountId;
+	FAccountId LocalUserInvalid;
+	TArray< FAccountId*> AccountIds = { &AccountId, &LocalUserInvalid };
+
+	LeaderboardsFixture_NUsers(std::move(GetLoginPipeline(AccountId)), AccountIds, 1)
+		.EmplaceAsyncLambda([&](FAsyncLambdaResult Promise, SubsystemType Services) { ReadEntriesAroundRankFixture(Promise, Services, "Stat_Use_Set", LocalUserInvalid, 0, 2, UE::Online::Errors::InvalidUser(), LeaderboardEntries); });
+
+	RunToCompletion();
+}
+
+
+LEADERBOARDS_TEST_CASE("Verify ReadEntriesAroundRank returns a fail message if given a Rank past the last existing rank")
+{
+	TArray<UE::Online::FLeaderboardEntry> LeaderboardEntries;
+	FAccountId AccountId;
+	TArray< FAccountId*> AccountIds = { &AccountId };
+
+	LeaderboardsFixture_NUsers(std::move(GetLoginPipeline(AccountId)), AccountIds, 1)
+		.EmplaceAsyncLambda([&](FAsyncLambdaResult Promise, SubsystemType Services) { ReadEntriesAroundRankFixture(Promise, Services, "Stat_Use_Set", AccountId, 1, 2, UE::Online::Errors::InvalidParams(), LeaderboardEntries); });
+
+	RunToCompletion();
+}
+
+
+LEADERBOARDS_TEST_CASE("Verify ReadEntriesAroundRank returns all valid entries if the given Limit exceeds the number of existing entries after the given Rank")
+{
+	TArray<UE::Online::FLeaderboardEntry> LeaderboardEntries;
+	FAccountId AccountIdA, AccountIdB;
+	TArray< FAccountId*> AccountIds = { &AccountIdA, &AccountIdB };
+	uint32 rank = 0;
+	uint32 limit = 3;
+
+	LeaderboardsFixture_NUsers(std::move(GetLoginPipeline(AccountIdA, AccountIdB)), AccountIds, 2)
+		.EmplaceAsyncLambda([&](FAsyncLambdaResult Promise, SubsystemType Services) { ReadEntriesAroundRank_2UsersFixture(Promise, Services, "Stat_Use_Set", AccountIdA, rank, limit, LeaderboardEntries); })
+		.EmplaceLambda([&](SubsystemType Services) { CheckEntries_2Users_Stat_Use_Set(LeaderboardEntries, AccountIdA, AccountIdB); });
+
+	RunToCompletion();
+}
+
+
+LEADERBOARDS_TEST_CASE("Verify ReadEntriesAroundRank returns fail message if given a Limit of 0")
+{
+	TArray<UE::Online::FLeaderboardEntry> LeaderboardEntries;
+	FAccountId AccountId;
+	TArray< FAccountId*> AccountIds = { &AccountId };
+
+	LeaderboardsFixture_NUsers(std::move(GetLoginPipeline(AccountId)), AccountIds, 1)
+		.EmplaceAsyncLambda([&](FAsyncLambdaResult Promise, SubsystemType Services) { ReadEntriesAroundRankFixture(Promise, Services, "Stat_Use_Set", AccountId, 3, 0, UE::Online::Errors::InvalidParams(), LeaderboardEntries); });
+
+	RunToCompletion();
+}
+
+LEADERBOARDS_TEST_CASE("Verify ReadEntriesAroundRank returns a fail message if given an invalid BoardName")
+{
+	// TODO
+}
\ No newline at end of file
diff --git a/Engine/Source/Programs/Online/OnlineTestsCore/Source/Private/Tests/Leaderboards/LeaderboardsReadEntriesAroundUserTests.cpp b/Engine/Source/Programs/Online/OnlineTestsCore/Source/Private/Tests/Leaderboards/LeaderboardsReadEntriesAroundUserTests.cpp
new file mode 100644
--- /dev/null
+++ b/Engine/Source/Programs/Online/OnlineTestsCore/Source/Private/Tests/Leaderboards/LeaderboardsReadEntriesAroundUserTests.cpp
@@ -0,0 +1,184 @@
+// Copyright Epic Games, Inc. All Rights Reserved.
+
+#include "ReadEntriesAroundUserFixture.h"
+#include "LeaderboardsFixture.h"
+
+#define LEADERBOARDS_TAGS "[Leaderboards]"
+#define LEADERBOARDS_TEST_CASE(x, ...) ONLINE_TEST_CASE(x, LEADERBOARDS_TAGS __VA_ARGS__)
+
+
+// ReadEntriesAroundUser tests
+
+LEADERBOARDS_TEST_CASE("Verify that ReadEntriesAroundUser succeeds and returns entries for the given board", "[MultiAccount]")
+{
+	DestroyCurrentServiceModule();
+
+	FAccountId AccountIdA, AccountIdB;
+	TArray< FAccountId*> AccountIds = { &AccountIdA, &AccountIdB };
+	TArray<UE::Online::FLeaderboardEntry> LeaderboardEntries;
+	int32 Offset = 0;
+	uint32 Limit = 2;
+
+	LeaderboardsFixture_NUsers(std::move(GetLoginPipeline(AccountIdA, AccountIdB)), AccountIds, 2)
+		.EmplaceAsyncLambda([&](FAsyncLambdaResult Promise, SubsystemType Services) { ReadEntriesAroundUser_2UsersFixture(Promise, Services, "Stat_Use_Set", AccountIdB, AccountIdA, Offset, Limit, LeaderboardEntries); })
+		.EmplaceLambda([&](SubsystemType Services) { CheckEntries_2Users_Stat_Use_Set(LeaderboardEntries, AccountIdA, AccountIdB); })
+		.EmplaceAsyncLambda([&](FAsyncLambdaResult Promise, SubsystemType Services) { ReadEntriesAroundUser_2UsersFixture(Promise, Services, "Stat_Use_Smallest", AccountIdA, AccountIdB, -1, 2, LeaderboardEntries); })
+		.EmplaceLambda([&](SubsystemType Services) { CheckEntries_2Users_Stat_Use_Smallest(LeaderboardEntries, AccountIdA, AccountIdB); });
+
+	RunToCompletion();
+}
+
+LEADERBOARDS_TEST_CASE("Verify ReadEntriesAroundUser returns a fail message if the local user is not logged in", "[MultiAccount]")
+{
+	FAccountId LocalUserInvalid;
+	FAccountId AccountId;
+	TArray< FAccountId*> AccountIds = { &AccountId, &LocalUserInvalid };
+	TArray<UE::Online::FLeaderboardEntry> LeaderboardEntries;
+	int32 offset = 0;
+	uint32 limit = 2;
+
+	LeaderboardsFixture_NUsers(std::move(GetLoginPipeline(AccountId)), AccountIds, 1)
+		.EmplaceAsyncLambda([&](FAsyncLambdaResult Promise, SubsystemType Services) {ReadEntriesAroundUserFixture(Promise, Services, "Stat_Use_Set", LocalUserInvalid, AccountId, offset, limit, UE::Online::Errors::InvalidUser(), LeaderboardEntries); });
+
+	RunToCompletion();
+}
+
+LEADERBOARDS_TEST_CASE("Verify ReadEntriesAroundUser returns a fail message if the given UserId is invalid", "[MultiAccount]")
+{
+	FAccountId LocalUser;
+	FAccountId AccountIdInvalid;
+	TArray< FAccountId*> AccountIds = { &LocalUser, &AccountIdInvalid };
+	TArray<UE::Online::FLeaderboardEntry> LeaderboardEntries;
+	int32 offset = 0;
+	uint32 limit = 2;
+
+	LeaderboardsFixture_NUsers(std::move(GetLoginPipeline(LocalUser)), AccountIds, 1)
+		.EmplaceAsyncLambda([&](FAsyncLambdaResult Promise, SubsystemType Services) {ReadEntriesAroundUserFixture(Promise, Services, "Stat_Use_Set", LocalUser, AccountIdInvalid, offset, limit, UE::Online::Errors::InvalidUser(), LeaderboardEntries); });
+
+	RunToCompletion();
+}
+
+LEADERBOARDS_TEST_CASE("Verify ReadEntriesAround User returns entries centered around the correct position when the given Offset is positive", "[MultiAccount]")
+{
+	FAccountId LocalUser, OtherPlayerB, UserPlayer, OtherPlayerC, OtherPlayerD;
+	TArray< FAccountId*> AccountIds = { &LocalUser, &OtherPlayerB, &UserPlayer, &OtherPlayerC, &OtherPlayerD };
+	TArray<UE::Online::FLeaderboardEntry> LeaderboardEntries;
+	int32 offset = 1;
+	uint32 limit = 2;
+
+	LeaderboardsFixture_NUsers(std::move(GetLoginPipeline(LocalUser, OtherPlayerB, UserPlayer, OtherPlayerC, OtherPlayerD)), AccountIds, 5)
+		.EmplaceAsyncLambda([&](FAsyncLambdaResult Promise, SubsystemType Services) { ReadEntriesAroundUser_5UsersFixture(Promise, Services, "Stat_Use_Set", { LocalUser, OtherPlayerB, UserPlayer, OtherPlayerC, OtherPlayerD }, offset, limit, LeaderboardEntries); })
+		.EmplaceLambda([&](SubsystemType Services) { CheckEntries_2UsersRanks_Stat_Use_Set(LeaderboardEntries, OtherPlayerC, OtherPlayerD); });
+
+	RunToCompletion();
+}
+
+LEADERBOARDS_TEST_CASE("Verify ReadEntriesAroundUser returns a number of entries that matches the given Limit", "[MultiAccount]")
+{
+
+	FAccountId LocalUser, AccountId;
+	TArray< FAccountId*> AccountIds = { &AccountId, &LocalUser };
+	int32 offset = 0;
+	uint32 limit = 2;
+
+	LeaderboardsFixture_NUsers(std::move(GetLoginPipeline(LocalUser, AccountId)), AccountIds, 2)
+		.EmplaceAsyncLambda([&](FAsyncLambdaResult Promise, SubsystemType Services) {
+		Services->GetLeaderboardsInterface()->ReadEntriesAroundUser({ LocalUser, AccountId, offset, limit, "Stat_Use_Set" })
+			.OnComplete([Promise, limit](const TOnlineResult<UE::Online::FReadEntriesAroundUser>& Result) mutable
+				{
+					CHECK_EQUAL(Result.GetOkValue().Entries.Num(), limit);
+					Promise->SetValue(true);
+				});
+			});
+
+	RunToCompletion();
+}
+
+LEADERBOARDS_TEST_CASE("Verify ReadEntriesAroundUser returns expected valid entries if the given Limit exceeds the number of existing entries before the center entry", "[MultiAccount]")
+{
+	DestroyCurrentServiceModule();
+
+	FAccountId LocalUser, AccountId;
+	TArray< FAccountId*> AccountIds = { &AccountId, &LocalUser };
+	int32 offset = -1;
+	uint32 limit = 4;
+	TArray<UE::Online::FLeaderboardEntry> LeaderboardEntries;
+
+	LeaderboardsFixture_NUsers(std::move(GetLoginPipeline(AccountId, LocalUser)), AccountIds, 2)
+		.EmplaceAsyncLambda([&](FAsyncLambdaResult Promise, SubsystemType Services) { ReadEntriesAroundUser_2UsersFixture(Promise, Services, "Stat_Use_Set", LocalUser, AccountId, offset, limit, LeaderboardEntries); })
+		.EmplaceLambda([&](SubsystemType Services) { CheckEntries_2Users_Stat_Use_Set(LeaderboardEntries, AccountId, LocalUser); });
+
+	RunToCompletion();
+}
+
+LEADERBOARDS_TEST_CASE("Verify ReadEntriesAroundUser returns expected valid entries if the given Limit exceeds the number of existing entries after the center entry", "[MultiAccount]")
+{
+	DestroyCurrentServiceModule();
+
+	FAccountId LocalUser, AccountId;
+	TArray< FAccountId*> AccountIds = { &AccountId, &LocalUser };
+	int32 offset = 0;
+	uint32 limit = 8;
+	TArray<UE::Online::FLeaderboardEntry> LeaderboardEntries;
+
+	LeaderboardsFixture_NUsers(std::move(GetLoginPipeline(LocalUser, AccountId)), AccountIds, 2)
+		.EmplaceAsyncLambda([&](FAsyncLambdaResult Promise, SubsystemType Services) { ReadEntriesAroundUser_2UsersFixture(Promise, Services, "Stat_Use_Smallest", LocalUser, AccountId, offset, limit, LeaderboardEntries); })
+		.EmplaceLambda([&](SubsystemType Services) { CheckEntries_1User_Stat_Use_Smallest(LeaderboardEntries, AccountId); });
+
+	RunToCompletion();
+}
+
+LEADERBOARDS_TEST_CASE("Verify ReadEntriesAroundUser returns error message if given a Limit of 0", "[MultiAccount]")
+{
+	FAccountId LocalUser, AccountId;
+	TArray< FAccountId*> AccountIds = { &AccountId, &LocalUser };
+	int32 offset = 0;
+	uint32 limit = 0;
+	TArray<UE::Online::FLeaderboardEntry> LeaderboardEntries;
+
+	LeaderboardsFixture_NUsers(std::move(GetLoginPipeline(LocalUser, AccountId)), AccountIds, 2)
+		.EmplaceAsyncLambda([&](FAsyncLambdaResult Promise, SubsystemType Services) {ReadEntriesAroundUserFixture(Promise, Services, "Stat_Use_Set", LocalUser, AccountId, offset, limit, UE::Online::Errors::InvalidParams(), LeaderboardEntries); });
+
+	RunToCompletion();
+}
+
+LEADERBOARDS_TEST_CASE("Verify ReadEntriesAroundUser returns only one entry if given Offset of 0 and  Limit of 1", "[MultiAccount]")
+{
+	FAccountId LocalUser, AccountId;
+	TArray< FAccountId*> AccountIds = { &AccountId, &LocalUser };
+	int32 offset = 0;
+	uint32 limit = 1;
+
+	LeaderboardsFixture_NUsers(std::move(GetLoginPipeline(LocalUser, AccountId)), AccountIds, 2)
+		.EmplaceAsyncLambda([&](FAsyncLambdaResult Promise, SubsystemType Services) {
+		Services->GetLeaderboardsInterface()->ReadEntriesAroundUser({ LocalUser, AccountId, offset, limit, "Stat_Use_Set" })
+			.OnComplete([Promise, limit](const TOnlineResult<UE::Online::FReadEntriesAroundUser>& Result) mutable
+				{
+					CHECK_EQUAL(Result.GetOkValue().Entries.Num(), limit);
+					Promise->SetValue(true);
+				});
+			});
+
+	RunToCompletion();
+}
+
+LEADERBOARDS_TEST_CASE("Verify ReadEntriesAroundUser returns expected valid entries if  given the Offset=-1 and Limit=2", "[MultiAccount]")
+{
+	FAccountId OtherPlayerA, OtherPlayerB, UserPlayer, OtherPlayerC, OtherPlayerD;
+	TArray< FAccountId*> AccountIds = { &OtherPlayerA, &OtherPlayerB, &UserPlayer, &OtherPlayerC, &OtherPlayerD };
+	TArray<UE::Online::FLeaderboardEntry> LeaderboardEntries;
+	int32 offset = 0;
+	uint32 limit = 2;
+
+	LeaderboardsFixture_NUsers(std::move(GetLoginPipeline(OtherPlayerA, OtherPlayerB, UserPlayer, OtherPlayerC, OtherPlayerD)), AccountIds, 5)
+		.EmplaceAsyncLambda([&](FAsyncLambdaResult Promise, SubsystemType Services) { ReadEntriesAroundUser_5UsersFixture(Promise, Services, "Stat_Use_Set", { OtherPlayerA, OtherPlayerB, UserPlayer, OtherPlayerC, OtherPlayerD }, offset, limit, LeaderboardEntries); })
+		.EmplaceLambda([&](SubsystemType Services) { CheckEntries_2UsersRanks_Stat_Use_Set_Index_Moved_Up(LeaderboardEntries, UserPlayer, OtherPlayerC); });
+
+	RunToCompletion();
+}
+
+LEADERBOARDS_TEST_CASE("Verify ReadEntriesAroundUser returns a fail message if given an invalid BoardName")
+{
+	// TODO
+}
+
diff --git a/Engine/Source/Programs/Online/OnlineTestsCore/Source/Private/Tests/Leaderboards/LeaderboardsReadEntriesForUsersTests.cpp b/Engine/Source/Programs/Online/OnlineTestsCore/Source/Private/Tests/Leaderboards/LeaderboardsReadEntriesForUsersTests.cpp
new file mode 100644
--- /dev/null
+++ b/Engine/Source/Programs/Online/OnlineTestsCore/Source/Private/Tests/Leaderboards/LeaderboardsReadEntriesForUsersTests.cpp
@@ -0,0 +1,118 @@
+// Copyright Epic Games, Inc. All Rights Reserved.
+
+#include "ReadEntiresForUsersFixture.h"
+#include "LeaderboardsFixture.h"
+
+#define LEADERBOARDS_TAGS "[Leaderboards]"
+#define LEADERBOARDS_TEST_CASE(x, ...) ONLINE_TEST_CASE(x, LEADERBOARDS_TAGS __VA_ARGS__)
+
+
+// ReadEntriesForUsers tests
+
+LEADERBOARDS_TEST_CASE("ReadEntriesForUsers succeed", "[MultiAccount]")
+{
+	DestroyCurrentServiceModule();
+
+	TArray<UE::Online::FLeaderboardEntry> LeaderboardEntries;
+	FAccountId AccountIdA, AccountIdB;
+	TArray< FAccountId*> AccountIds = { &AccountIdA, &AccountIdB };
+
+	LeaderboardsFixture_NUsers(std::move(GetLoginPipeline(AccountIdA, AccountIdB)), AccountIds, 2)
+		.EmplaceAsyncLambda([&](FAsyncLambdaResult Promise, SubsystemType Services) { ReadEntriesForUsers_2UsersFixture(Promise, Services, "Stat_Use_Set", AccountIdA, AccountIdB, LeaderboardEntries); })
+		.EmplaceLambda([&](SubsystemType Services) { CheckEntries_2Users_Stat_Use_Set(LeaderboardEntries, AccountIdA, AccountIdB); })
+		.EmplaceAsyncLambda([&](FAsyncLambdaResult Promise, SubsystemType Services) { ReadEntriesForUsers_2UsersFixture(Promise, Services, "Stat_Use_Smallest", AccountIdA, AccountIdB, LeaderboardEntries); })
+		.EmplaceLambda([&](SubsystemType Services) { CheckEntries_2Users_Stat_Use_Smallest(LeaderboardEntries, AccountIdA, AccountIdB); });
+
+	RunToCompletion();
+}
+
+LEADERBOARDS_TEST_CASE("Verify ReadEntriesForUsers returns a fail message if the local user is not logged in", "[MultiAccount]")
+{
+	TArray<UE::Online::FLeaderboardEntry> LeaderboardEntries;
+	FAccountId AccountId;	
+	FAccountId LocalUserInvalid;
+	TArray< FAccountId*> AccountIds = { &AccountId, &LocalUserInvalid };
+	LeaderboardsFixture_NUsers(std::move(GetLoginPipeline(AccountId)), AccountIds, 2)
+		.EmplaceAsyncLambda([&](FAsyncLambdaResult Promise, SubsystemType Services) { ReadEntriesForUsersFixture(Promise, Services, "Stat_Use_Set", { LocalUserInvalid, AccountId }, UE::Online::Errors::InvalidUser(), LeaderboardEntries); });
+
+	RunToCompletion();
+}
+
+LEADERBOARDS_TEST_CASE("Verify ReadEntriesForUsers returns an empty array of entries if the given UserIds array is empty")
+{
+	TArray<UE::Online::FLeaderboardEntry> LeaderboardEntries;
+	FAccountId LocalUser;
+	TArray< FAccountId*> AccountIds = { &LocalUser };
+	LeaderboardsFixture_NUsers(std::move(GetLoginPipeline(LocalUser)), AccountIds, 1)
+		.EmplaceAsyncLambda([&](FAsyncLambdaResult Promise, SubsystemType Services) {
+		Services->GetLeaderboardsInterface()->ReadEntriesForUsers({ LocalUser, {}, "Stat_Use_Set" })
+		.OnComplete([Promise, &LeaderboardEntries](const TOnlineResult<UE::Online::FReadEntriesForUsers>& Result) mutable
+			{
+				CHECK(LeaderboardEntries.IsEmpty());
+				Promise->SetValue(true);
+			});
+		});
+	RunToCompletion();
+}
+
+LEADERBOARDS_TEST_CASE("Verify ReadEntiresForUsers returns a fail message if given a UserIds array that contains one user ID which is invalid", "[MultiAccount]")
+{
+	TArray<UE::Online::FLeaderboardEntry> LeaderboardEntries;
+	FAccountId AccountIdInvalid;
+	FAccountId LocalUser;
+	TArray< FAccountId*> AccountIds = { &LocalUser, &AccountIdInvalid };
+	LeaderboardsFixture_NUsers(std::move(GetLoginPipeline(LocalUser)), AccountIds, 1)
+		.EmplaceAsyncLambda([&](FAsyncLambdaResult Promise, SubsystemType Services) {
+		Services->GetLeaderboardsInterface()->ReadEntriesForUsers({ LocalUser, {AccountIdInvalid}, "Stat_Use_Set" })
+		.OnComplete([Promise, &LeaderboardEntries](const TOnlineResult<UE::Online::FReadEntriesForUsers>& Result) mutable
+			{
+				if (Result.IsError())
+				{
+					CHECK(Result.GetErrorValue() == UE::Online::Errors::InvalidUser());
+					CHECK(LeaderboardEntries.IsEmpty());
+				}
+				else
+				{
+					CHECK(LeaderboardEntries.IsEmpty());
+				}
+				Promise->SetValue(true);
+			});
+		});
+	RunToCompletion();
+}
+
+LEADERBOARDS_TEST_CASE("Verify ReadEntriesForUsers returns a fail message if given a UserIds array that contains one valid user ID and one invalid user ID", "[MultiAccount]")
+{
+	TArray<UE::Online::FLeaderboardEntry> LeaderboardEntries;
+	FAccountId AccountIdInvalid;
+	FAccountId LocalUser;
+	TArray< FAccountId*> AccountIds = { &LocalUser, &AccountIdInvalid };
+	LeaderboardsFixture_NUsers(std::move(GetLoginPipeline(LocalUser)), AccountIds, 1)
+		.EmplaceAsyncLambda([&](FAsyncLambdaResult Promise, SubsystemType Services) { ReadEntriesForUsersFixture(Promise, Services, "Stat_Use_Set", { LocalUser, AccountIdInvalid }, UE::Online::Errors::InvalidUser(), LeaderboardEntries); });
+	RunToCompletion();
+}
+
+LEADERBOARDS_TEST_CASE("Verify ReadEntriesForUsers returns an Entries array with only the expected entry if given a UserIds array with one ID")
+{
+	TArray<UE::Online::FLeaderboardEntry> LeaderboardEntries;
+	FAccountId LocalUser;
+	TArray< FAccountId*> AccountIds = { &LocalUser };
+	LeaderboardsFixture_NUsers(std::move(GetLoginPipeline(LocalUser)), AccountIds, 1)
+		.EmplaceAsyncLambda([&](FAsyncLambdaResult Promise, SubsystemType Services) {
+		Services->GetLeaderboardsInterface()->ReadEntriesForUsers({ LocalUser, {LocalUser}, "Stat_Use_Set" })
+		.OnComplete([Promise, &LeaderboardEntries, &LocalUser](const TOnlineResult<UE::Online::FReadEntriesForUsers>& Result) mutable
+			{
+				CHECK_EQUAL(Result.GetOkValue().Entries.Num(), 1);
+				CHECK_EQUAL(LeaderboardEntryOf(Result.GetOkValue().Entries, LocalUser)->Rank, 1);
+				CHECK_EQUAL(LeaderboardEntryOf(Result.GetOkValue().Entries, LocalUser)->Score, 3);
+				Promise->SetValue(true);
+			});
+		});
+	RunToCompletion();
+}
+
+LEADERBOARDS_TEST_CASE("Verify ReadEntriesForUsers returns a fail message if given an invalid BoardName")
+{
+	// TODO
+}
+
diff --git a/Engine/Source/Programs/Online/OnlineTestsCore/Source/Private/Tests/Leaderboards/ReadEntiresForUsersFixture.cpp b/Engine/Source/Programs/Online/OnlineTestsCore/Source/Private/Tests/Leaderboards/ReadEntiresForUsersFixture.cpp
new file mode 100644
--- /dev/null
+++ b/Engine/Source/Programs/Online/OnlineTestsCore/Source/Private/Tests/Leaderboards/ReadEntiresForUsersFixture.cpp
@@ -0,0 +1,32 @@
+// Copyright Epic Games, Inc. All Rights Reserved.
+
+#include "ReadEntiresForUsersFixture.h"
+
+
+void ReadEntriesForUsers_2UsersFixture(FAsyncLambdaResult Promise, SubsystemType Services, const FString& BoardName, const FAccountId& AccountIdA, const FAccountId& AccountIdB, TArray<UE::Online::FLeaderboardEntry>& OutLeaderboardEntries)
+{
+	Services->GetLeaderboardsInterface()->ReadEntriesForUsers({ AccountIdA, {AccountIdA, AccountIdB}, BoardName }).OnComplete([Promise, &OutLeaderboardEntries](const TOnlineResult<UE::Online::FReadEntriesForUsers>& Result) mutable
+		{
+			CHECK_OP(Result);
+			Promise->SetValue(true);
+			OutLeaderboardEntries = Result.GetOkValue().Entries;
+		});
+}
+
+void ReadEntriesForUsersFixture(FAsyncLambdaResult Promise, SubsystemType Services, const FString& BoardName, const TArray<FAccountId>& AccountIds, const UE::Online::FOnlineError& InError, TArray<UE::Online::FLeaderboardEntry>& OutLeaderboardEntries)
+{
+	Services->GetLeaderboardsInterface()->ReadEntriesForUsers({ AccountIds[0], AccountIds, BoardName }).OnComplete([Promise, &OutLeaderboardEntries, InError](const TOnlineResult<UE::Online::FReadEntriesForUsers>& Result) mutable
+		{
+			//CHECK_OP_EQ(Result, InError);
+			if (Result.IsError())
+			{
+				CHECK(Result.GetErrorValue() == InError);
+				CHECK(OutLeaderboardEntries.IsEmpty());
+			}
+			else
+			{
+				CHECK(OutLeaderboardEntries.IsEmpty());
+			}
+				Promise->SetValue(true);
+		});
+}
\ No newline at end of file
diff --git a/Engine/Source/Programs/Online/OnlineTestsCore/Source/Private/Tests/Leaderboards/ReadEntiresForUsersFixture.h b/Engine/Source/Programs/Online/OnlineTestsCore/Source/Private/Tests/Leaderboards/ReadEntiresForUsersFixture.h
new file mode 100644
--- /dev/null
+++ b/Engine/Source/Programs/Online/OnlineTestsCore/Source/Private/Tests/Leaderboards/ReadEntiresForUsersFixture.h
@@ -0,0 +1,15 @@
+// Copyright Epic Games, Inc. All Rights Reserved.
+
+#include <catch2/catch_test_macros.hpp>
+
+#include "Helpers/Identity/IdentityGetLoginByUserId.h"
+#include "Helpers/Stats/UpdateStatsHelper.h"
+#include "AsyncTestStep.h"
+#include "OnlineCatchHelper.h"
+#include "Online/Stats.h"
+#include "Online/Leaderboards.h"
+
+
+void ReadEntriesForUsers_2UsersFixture(FAsyncLambdaResult Promise, SubsystemType Services, const FString& BoardName, const FAccountId& AccountIdA, const FAccountId& AccountIdB, TArray<UE::Online::FLeaderboardEntry>& OutLeaderboardEntries);
+void ReadEntriesForUsersFixture(FAsyncLambdaResult Promise, SubsystemType Services, const FString& BoardName, const TArray<FAccountId>& AccountIds, const UE::Online::FOnlineError& InError, TArray<UE::Online::FLeaderboardEntry>& OutLeaderboardEntries);
+
diff --git a/Engine/Source/Programs/Online/OnlineTestsCore/Source/Private/Tests/Leaderboards/ReadEntriesAroundRankFixture.cpp b/Engine/Source/Programs/Online/OnlineTestsCore/Source/Private/Tests/Leaderboards/ReadEntriesAroundRankFixture.cpp
new file mode 100644
--- /dev/null
+++ b/Engine/Source/Programs/Online/OnlineTestsCore/Source/Private/Tests/Leaderboards/ReadEntriesAroundRankFixture.cpp
@@ -0,0 +1,30 @@
+// Copyright Epic Games, Inc. All Rights Reserved.
+
+#include "ReadEntriesAroundRankFixture.h"
+
+void ReadEntriesAroundRank_2UsersFixture(FAsyncLambdaResult Promise, SubsystemType Services, const FString& BoardName, const FAccountId& AccountIdA, uint32 Rank, uint32 Limit, TArray<UE::Online::FLeaderboardEntry>& OutLeaderboardEntries)
+{
+	Services->GetLeaderboardsInterface()->ReadEntriesAroundRank({ AccountIdA, Rank, Limit, BoardName }).OnComplete([Promise, &OutLeaderboardEntries](const TOnlineResult<UE::Online::FReadEntriesAroundRank>& Result) mutable
+		{
+			CHECK_OP(Result);
+			Promise->SetValue(true);
+			OutLeaderboardEntries = Result.GetOkValue().Entries;
+		});
+}
+
+void ReadEntriesAroundRankFixture(FAsyncLambdaResult Promise, SubsystemType Services, const FString& BoardName, const FAccountId& AccountId, uint32 Rank, uint32 Limit, const UE::Online::FOnlineError& InError, TArray<UE::Online::FLeaderboardEntry>& OutLeaderboardEntries)
+{
+	Services->GetLeaderboardsInterface()->ReadEntriesAroundRank({ AccountId, Rank, Limit, BoardName }).OnComplete([Promise, &OutLeaderboardEntries, InError](const TOnlineResult<UE::Online::FReadEntriesAroundRank>& Result) mutable
+		{
+			if (Result.IsError())
+			{
+				CHECK(Result.GetErrorValue() == InError);
+				CHECK(OutLeaderboardEntries.IsEmpty());
+			}
+			else
+			{
+				CHECK(OutLeaderboardEntries.IsEmpty());
+			}
+	Promise->SetValue(true);
+		});
+}
\ No newline at end of file
diff --git a/Engine/Source/Programs/Online/OnlineTestsCore/Source/Private/Tests/Leaderboards/ReadEntriesAroundRankFixture.h b/Engine/Source/Programs/Online/OnlineTestsCore/Source/Private/Tests/Leaderboards/ReadEntriesAroundRankFixture.h
new file mode 100644
--- /dev/null
+++ b/Engine/Source/Programs/Online/OnlineTestsCore/Source/Private/Tests/Leaderboards/ReadEntriesAroundRankFixture.h
@@ -0,0 +1,14 @@
+// Copyright Epic Games, Inc. All Rights Reserved.
+
+#include <catch2/catch_test_macros.hpp>
+
+#include "Helpers/Identity/IdentityGetLoginByUserId.h"
+#include "Helpers/Stats/UpdateStatsHelper.h"
+#include "AsyncTestStep.h"
+#include "OnlineCatchHelper.h"
+#include "Online/Stats.h"
+#include "Online/Leaderboards.h"
+
+
+void ReadEntriesAroundRank_2UsersFixture(FAsyncLambdaResult Promise, SubsystemType Services, const FString& BoardName, const FAccountId& AccountIdA, uint32 Rank, uint32 Limit, TArray<UE::Online::FLeaderboardEntry>& OutLeaderboardEntries);
+void ReadEntriesAroundRankFixture(FAsyncLambdaResult Promise, SubsystemType Services, const FString& BoardName, const FAccountId& AccountId, uint32 Rank, uint32 Limit, const UE::Online::FOnlineError& InError, TArray<UE::Online::FLeaderboardEntry>& OutLeaderboardEntries);
diff --git a/Engine/Source/Programs/Online/OnlineTestsCore/Source/Private/Tests/Leaderboards/ReadEntriesAroundUserFixture.cpp b/Engine/Source/Programs/Online/OnlineTestsCore/Source/Private/Tests/Leaderboards/ReadEntriesAroundUserFixture.cpp
new file mode 100644
--- /dev/null
+++ b/Engine/Source/Programs/Online/OnlineTestsCore/Source/Private/Tests/Leaderboards/ReadEntriesAroundUserFixture.cpp
@@ -0,0 +1,40 @@
+// Copyright Epic Games, Inc. All Rights Reserved.
+
+#include "ReadEntriesAroundUserFixture.h"
+
+void ReadEntriesAroundUser_2UsersFixture(FAsyncLambdaResult Promise, SubsystemType Services, const FString& BoardName, const FAccountId& AccountIdA, const FAccountId& AccountIdB, int32 Offset, uint32 Limit, TArray<UE::Online::FLeaderboardEntry>& OutLeaderboardEntries)
+{
+	Services->GetLeaderboardsInterface()->ReadEntriesAroundUser({ AccountIdA, AccountIdB, Offset, Limit, BoardName }).OnComplete([Promise, &OutLeaderboardEntries](const TOnlineResult<UE::Online::FReadEntriesAroundUser>& Result) mutable
+		{
+			CHECK_OP(Result);
+			Promise->SetValue(true);
+			OutLeaderboardEntries = Result.GetOkValue().Entries;
+		});
+}
+
+void ReadEntriesAroundUser_5UsersFixture(FAsyncLambdaResult Promise, SubsystemType Services, const FString& BoardName, const TArray<FAccountId>& AccountIds, int32 Offset, uint32 Limit, TArray<UE::Online::FLeaderboardEntry>& OutLeaderboardEntries)
+{
+	Services->GetLeaderboardsInterface()->ReadEntriesAroundUser({ AccountIds[0], AccountIds[2], Offset, Limit, BoardName }).OnComplete([Promise, &OutLeaderboardEntries](const TOnlineResult<UE::Online::FReadEntriesAroundUser>& Result) mutable
+		{
+			CHECK(Result.IsOk());
+			Promise->SetValue(true);
+			OutLeaderboardEntries = Result.GetOkValue().Entries;
+		});
+}
+
+void ReadEntriesAroundUserFixture(FAsyncLambdaResult Promise, SubsystemType Services, const FString& BoardName, const FAccountId& AccountIdA, const FAccountId& AccountIdB, int32 Offset, uint32 Limit, const UE::Online::FOnlineError& InError, TArray<UE::Online::FLeaderboardEntry>& OutLeaderboardEntries)
+{
+	Services->GetLeaderboardsInterface()->ReadEntriesAroundUser({ AccountIdA, AccountIdB, Offset, Limit, BoardName }).OnComplete([Promise, &OutLeaderboardEntries, InError](const TOnlineResult<UE::Online::FReadEntriesAroundUser>& Result) mutable
+		{
+			if (Result.IsError())
+			{
+				CHECK(Result.GetErrorValue() == InError);
+				CHECK(OutLeaderboardEntries.IsEmpty());
+			}
+			else
+			{
+				CHECK(OutLeaderboardEntries.IsEmpty());
+			}
+			Promise->SetValue(true);
+		});
+}
\ No newline at end of file
diff --git a/Engine/Source/Programs/Online/OnlineTestsCore/Source/Private/Tests/Leaderboards/ReadEntriesAroundUserFixture.h b/Engine/Source/Programs/Online/OnlineTestsCore/Source/Private/Tests/Leaderboards/ReadEntriesAroundUserFixture.h
new file mode 100644
--- /dev/null
+++ b/Engine/Source/Programs/Online/OnlineTestsCore/Source/Private/Tests/Leaderboards/ReadEntriesAroundUserFixture.h
@@ -0,0 +1,15 @@
+// Copyright Epic Games, Inc. All Rights Reserved.
+
+#include <catch2/catch_test_macros.hpp>
+
+#include "Helpers/Identity/IdentityGetLoginByUserId.h"
+#include "Helpers/Stats/UpdateStatsHelper.h"
+#include "AsyncTestStep.h"
+#include "OnlineCatchHelper.h"
+#include "Online/Stats.h"
+#include "Online/Leaderboards.h"
+
+
+void ReadEntriesAroundUser_2UsersFixture(FAsyncLambdaResult Promise, SubsystemType Services, const FString& BoardName, const FAccountId& AccountIdA, const FAccountId& AccountIdB, int32 Offset, uint32 Limit, TArray<UE::Online::FLeaderboardEntry>& OutLeaderboardEntries);
+void ReadEntriesAroundUser_5UsersFixture(FAsyncLambdaResult Promise, SubsystemType Services, const FString& BoardName, const TArray<FAccountId>& AccountIds, int32 Offset, uint32 Limit, TArray<UE::Online::FLeaderboardEntry>& OutLeaderboardEntries);
+void ReadEntriesAroundUserFixture(FAsyncLambdaResult Promise, SubsystemType Services, const FString& BoardName, const FAccountId& AccountIdA, const FAccountId& AccountIdB, int32 Offset, uint32 Limit, const UE::Online::FOnlineError& InError, TArray<UE::Online::FLeaderboardEntry>& OutLeaderboardEntries);
diff --git a/Engine/Source/Programs/Online/OnlineTestsCore/Source/Private/Tests/Lobby/LobbyCreateInviteTests.cpp b/Engine/Source/Programs/Online/OnlineTestsCore/Source/Private/Tests/Lobby/LobbyCreateInviteTests.cpp
new file mode 100644
--- /dev/null
+++ b/Engine/Source/Programs/Online/OnlineTestsCore/Source/Private/Tests/Lobby/LobbyCreateInviteTests.cpp
@@ -0,0 +1,32 @@
+// Copyright Epic Games, Inc. All Rights Reserved.
+
+#include "CoreMinimal.h"
+
+#include <catch2/catch_test_macros.hpp>
+
+#include "Helpers/Identity/IdentityGetLoginByUserId.h"
+#include "Helpers/Lobby/LobbyCreateHelper.h"
+#include "OnlineCatchHelper.h"
+
+#define LOBBY_CREATE_INVITE_TAGS "[Lobby]"
+#define LOBBY_CREATE_INVITE_TEST_CASE(x, ...) ONLINE_TEST_CASE(x, LOBBY_CREATE_INVITE_TAGS __VA_ARGS__)
+
+LOBBY_CREATE_INVITE_TEST_CASE("Basic lobby create test")
+{
+	int32 LocalUserNum = 0;
+	UE::Online::FCreateLobby::Params Params;
+	Params.LocalName = TEXT("TestLobby");
+	Params.SchemaId = TEXT("test");
+	Params.MaxMembers = 2;
+	Params.JoinPolicy = ELobbyJoinPolicy::PublicAdvertised;
+	Params.Attributes = { { TEXT("LobbyCreateTime"), (int64)10}};
+
+	GetLoginPipeline(1)
+		.EmplaceStep<FIdentityGetUniquePlayerIdStep>(LocalUserNum, 
+			[&Params](UE::Online::FAccountInfo InAccountInfo) {
+				Params.LocalAccountId = InAccountInfo.AccountId;
+			})
+		.EmplaceStep<FLobbyCreateHelper>(&Params, true);
+
+	RunToCompletion();
+}
diff --git a/Engine/Source/Programs/Online/OnlineTestsCore/Source/Private/Tests/Social/SocialTests.cpp b/Engine/Source/Programs/Online/OnlineTestsCore/Source/Private/Tests/Social/SocialTests.cpp
new file mode 100644
--- /dev/null
+++ b/Engine/Source/Programs/Online/OnlineTestsCore/Source/Private/Tests/Social/SocialTests.cpp
@@ -0,0 +1,216 @@
+// Copyright Epic Games, Inc. All Rights Reserved.
+#include "CoreMinimal.h"
+#include "TestHarness.h"
+#include "OnlineCatchHelper.h"
+
+#define SOCIAL_TAG "[SOCIAL]"  
+#define SOCIAL_TEST_CASE(x, ...) ONLINE_TEST_CASE(x, SOCIAL_TAG __VA_ARGS__)
+
+// QueryFriends tests
+SOCIAL_TEST_CASE("Verify that QueryFriends returns a fail message if the local user is not logged in")
+{
+	// TODO
+}
+
+SOCIAL_TEST_CASE("Verify that QueryFriends caches no Friends if no Friends exist for this user")
+{
+	// TODO
+}
+
+SOCIAL_TEST_CASE("Verify that QueryFriends caches one Friend if only one Friend exists for this user")
+{
+	// TODO
+}
+
+SOCIAL_TEST_CASE("Verify that QueryFriends caches all Friends if multiple Friends exist for this user")
+{
+	// TODO
+}
+
+// GetFriends tests
+SOCIAL_TEST_CASE("Verify that GetFriends returns a fail message if the local user is not logged in")
+{
+	// TODO
+}
+
+SOCIAL_TEST_CASE("Verify that GetFriends returns an empty list if there are no cached Friends")
+{
+	// TODO
+}
+
+SOCIAL_TEST_CASE("Verify that GetFriends returns a list of 1 Friend if there is 1 cached Friend")
+{
+	// TODO
+}
+
+SOCIAL_TEST_CASE("Verify that GetFriends returns a list of all cached Friends if there are multiple cached Friends")
+{
+	// TODO
+}
+
+// GetFriend tests
+SOCIAL_TEST_CASE("Verify that GetFriend returns a fail message if the local user is not logged in")
+{
+	// TODO
+}
+
+SOCIAL_TEST_CASE("Verify that GetFriend returns a fail message if there are no cached Friends")
+{
+	// TODO
+}
+
+SOCIAL_TEST_CASE("Verify that GetFriend returns the correct Friend when given an existing FriendId")
+{
+	// TODO
+}
+
+SOCIAL_TEST_CASE("Verify that GetFriend returns a fail message when there are cached Friends but the given a FriendId that matches none of them")
+{
+	// TODO
+}
+
+// SendFriendInvite tests
+SOCIAL_TEST_CASE("Verify that SendFriendInvite returns a fail message if the local user is not logged in")
+{
+	// TODO
+}
+
+SOCIAL_TEST_CASE("Verify that SendFriendInvite returns a fail message if the given TargetUserId does not exist")
+{
+	// TODO
+}
+
+SOCIAL_TEST_CASE("Verify that SendFriendInvite returns fail message if ERelationship with target user is Friend")
+{
+	// TODO
+}
+
+SOCIAL_TEST_CASE("Verify that SendFriendInvite completes successfully if ERelationship with target user is NotFriend, ERelationship becomes InviteSent")
+{
+	// TODO
+}
+
+SOCIAL_TEST_CASE("Verify that SendFriendInvite completes successfully if ERelationship with target user is InviteSent, ERelationship remains InviteSent")
+{
+	// TODO
+}
+
+SOCIAL_TEST_CASE("Verify that SendFriendInvite completes successfully if ERelationship with target user is InviteReceived, ERelationship becomes InviteSent")
+{
+	// TODO
+}
+
+SOCIAL_TEST_CASE("Verify that SendFriendInvite returns fail message if ERelationship with target user is Blocked")
+{
+	// TODO
+}
+
+// AcceptFriendInvite tests
+SOCIAL_TEST_CASE("Verify that AcceptFriendInvite returns a fail message if the local user is not logged in")
+{
+	// TODO
+}
+
+SOCIAL_TEST_CASE("Verify that AcceptFriendInvite returns a fail message if ERelationship with target user is Friend")
+{
+	// TODO
+}
+
+SOCIAL_TEST_CASE("Verify that AcceptFriendInvite returns a fail message if ERelationship with target user is NotFriend")
+{
+	// TODO
+}
+
+SOCIAL_TEST_CASE("Verify that AcceptFriendInvite returns a fail message if ERelationship with target user is InviteSent")
+{
+	// TODO
+}
+
+SOCIAL_TEST_CASE("Verify that AcceptFriendInvite successfully completes if ERelationship with target user is InviteReceived, ERelationship becomes Friend")
+{
+	// TODO
+}
+
+SOCIAL_TEST_CASE("Verify that AcceptFriendInvite returns a fail message if ERelationship with target user is Blocked")
+{
+	// TODO
+}
+
+// RejectFriendInvite tests
+SOCIAL_TEST_CASE("Verify that RejectFriendInvite returns a fail message if the local user is not logged in")
+{
+	// TODO
+}
+
+SOCIAL_TEST_CASE("Verify that RejectFriendInvite returns a fail message if ERelationship with target user is Friend")
+{
+	// TODO
+}
+
+SOCIAL_TEST_CASE("Verify that RejectFriendInvite returns a fail message if ERelationship with target user is NotFriend")
+{
+	// TODO
+}
+
+SOCIAL_TEST_CASE("Verify that RejectFriendInvite returns a fail message if ERelationship with target user is InviteSent")
+{
+	// TODO
+}
+
+SOCIAL_TEST_CASE("Verify that RejectFriendInvite successfully completes if ERelationship with target user is InviteReceived, ERelationship becomes NotFriend")
+{
+	// TODO
+}
+
+SOCIAL_TEST_CASE("Verify that RejectFriendInvite returns a fail message if ERelationship with target user is Blocked")
+{
+	// TODO
+}
+
+// OnRelationshipUpdated tests
+// TODO: Unsure how to verify
+
+// QueryBlockedUsers tests
+SOCIAL_TEST_CASE("Verify that QueryBlockedUsers returns a fail message if the local user is not logged in")
+{
+	// TODO
+}
+
+SOCIAL_TEST_CASE("Verify that QueryBlockedUsers caches no blocked users if no blocked users exist for this user")
+{
+	// TODO
+}
+
+SOCIAL_TEST_CASE("Verify that QueryBlockedUsers caches one blocked user if only one blocked user exists for this user")
+{
+	// TODO
+}
+
+SOCIAL_TEST_CASE("Verify that QueryBlockedUsers caches all blocked users if multiple blocked users exist for this user")
+{
+	// TODO
+}
+
+// GetBlockedUsers tests
+SOCIAL_TEST_CASE("Verify that GetBlockedUsers returns a fail message if the local user is not logged in")
+{
+	// TODO
+}
+
+SOCIAL_TEST_CASE("Verify that GetBlockedUsers returns an empty list if there are no cached blocked users")
+{
+	// TODO
+}
+
+SOCIAL_TEST_CASE("Verify that GetBlockedUsers returns a list of 1 blocked user if there is 1 cached bocked user")
+{
+	// TODO
+}
+
+SOCIAL_TEST_CASE("Verify that GetBlockedUsers returns a list of all cached blocked user if there are multiple cached blocked users")
+{
+	// TODO
+}
+
+// OnBlockedUserUpdated tests
+// TODO: Unsure how to verify
\ No newline at end of file
diff --git a/Engine/Source/Programs/Online/OnlineTestsCore/Source/Private/Tests/TestAsyncOp.cpp b/Engine/Source/Programs/Online/OnlineTestsCore/Source/Private/Tests/TestAsyncOp.cpp
new file mode 100644
--- /dev/null
+++ b/Engine/Source/Programs/Online/OnlineTestsCore/Source/Private/Tests/TestAsyncOp.cpp
@@ -0,0 +1,392 @@
+// Copyright Epic Games, Inc. All Rights Reserved.
+#include "CoreMinimal.h"
+#include <catch2/catch_test_macros.hpp>
+
+#include "OnlineCatchHelper.h"
+
+#include "Online/OnlineUtils.h"
+#include "Online/OnlineResult.h"
+#include "Online/OnlineServicesCommon.h"
+
+#include "Async/Async.h"
+#include "Async/TaskGraphInterfaces.h"
+#include "Containers/Ticker.h"
+#include "HAL/PlatformTLS.h"
+
+#define ASYNCOP_TAG "[AsyncOp]"
+#define ASYNCOP_TEST_CASE(x, ...) ONLINE_TEST_CASE(x, ASYNCOP_TAG __VA_ARGS__)
+
+struct FTestOp
+{
+	static constexpr TCHAR Name[] = TEXT("TestOp");
+
+	struct Params {};
+	struct Result {};
+};
+
+struct FTestOpWithParams
+{
+	static constexpr TCHAR Name[] = TEXT("TestOpWithParams");
+
+	struct Params { int Index = -1; };
+	struct Result {};
+};
+
+struct FMergeableTestOpWithParams
+{
+	static constexpr TCHAR Name[] = TEXT("MergeableTestOpWithParams");
+
+	struct Params { int Index = -1; int Mutations = 0; };
+	struct Result {};
+};
+
+namespace UE::Online::Meta {
+
+BEGIN_ONLINE_STRUCT_META(FTestOp::Params)
+END_ONLINE_STRUCT_META()
+
+BEGIN_ONLINE_STRUCT_META(FTestOp::Result)
+END_ONLINE_STRUCT_META()
+
+BEGIN_ONLINE_STRUCT_META(FTestOpWithParams::Params)
+	ONLINE_STRUCT_FIELD(FTestOpWithParams::Params, Index)
+END_ONLINE_STRUCT_META()
+
+BEGIN_ONLINE_STRUCT_META(FTestOpWithParams::Result)
+END_ONLINE_STRUCT_META()
+
+BEGIN_ONLINE_STRUCT_META(FMergeableTestOpWithParams::Params)
+	ONLINE_STRUCT_FIELD(FMergeableTestOpWithParams::Params, Index),
+	ONLINE_STRUCT_FIELD(FMergeableTestOpWithParams::Params, Mutations)
+END_ONLINE_STRUCT_META()
+
+BEGIN_ONLINE_STRUCT_META(FMergeableTestOpWithParams::Result)
+END_ONLINE_STRUCT_META()
+
+/* namespace UE::Online::Meta */ }
+
+ASYNCOP_TEST_CASE("Verify that the cache returns the same joinable op when the parameters are the same")
+{
+	FTestPipeline& TestPipeline = GetPipeline();
+	TestPipeline.EmplaceAsyncLambda([](FAsyncLambdaResult Promise, SubsystemType Services)
+		{
+			UE::Online::FOnlineServicesCommon& ServicesCommon = static_cast<UE::Online::FOnlineServicesCommon&>(*Services);
+			FTestOpWithParams::Params Params1;
+			Params1.Index = 0;
+			TOnlineAsyncOpRef<FTestOpWithParams> Op = ServicesCommon.GetJoinableOp<FTestOpWithParams>(MoveTemp(Params1));
+			CHECK(!Op->IsReady());
+			Op->Then([Promise](FOnlineAsyncOp&) { Promise->SetValue(true); })
+				.Enqueue(ServicesCommon.GetParallelQueue());
+
+			FTestOpWithParams::Params Params2;
+			Params2.Index = 0;
+			TOnlineAsyncOpRef<FTestOpWithParams> Op2 = ServicesCommon.GetJoinableOp<FTestOpWithParams>(MoveTemp(Params2));
+			CHECK(Op2->IsReady());
+		});
+
+	RunToCompletion();
+}
+
+ASYNCOP_TEST_CASE("Verify that the cache returns different joinable ops when the parameters are different")
+{
+	FTestPipeline& TestPipeline = GetPipeline();
+	TestPipeline.EmplaceAsyncLambda([](FAsyncLambdaResult Promise, SubsystemType Services)
+		{
+			UE::Online::FOnlineServicesCommon& ServicesCommon = static_cast<UE::Online::FOnlineServicesCommon&>(*Services);
+			FTestOpWithParams::Params Params1;
+			Params1.Index = 1;
+			TOnlineAsyncOpRef<FTestOpWithParams> Op = ServicesCommon.GetJoinableOp<FTestOpWithParams>(MoveTemp(Params1));
+			CHECK(!Op->IsReady());
+			Op->Then([](FOnlineAsyncOp&) {})
+				.Enqueue(ServicesCommon.GetParallelQueue());
+
+			FTestOpWithParams::Params Params2;
+			Params2.Index = 2;
+			TOnlineAsyncOpRef<FTestOpWithParams> Op2 = ServicesCommon.GetJoinableOp<FTestOpWithParams>(MoveTemp(Params2));
+			CHECK(!Op2->IsReady());
+			Op2->Then([Promise](FOnlineAsyncOp&) { Promise->SetValue(true); })
+				.Enqueue(ServicesCommon.GetParallelQueue());
+		});
+
+	RunToCompletion();
+}
+
+ASYNCOP_TEST_CASE("Verify that the cache returns the same mergeable op when the parameters are the same")
+{
+	FTestPipeline& TestPipeline = GetPipeline();
+	TestPipeline.EmplaceAsyncLambda([](FAsyncLambdaResult Promise, SubsystemType Services)
+		{
+			UE::Online::FOnlineServicesCommon& ServicesCommon = static_cast<UE::Online::FOnlineServicesCommon&>(*Services);
+			FMergeableTestOpWithParams::Params Params1;
+			Params1.Index = 0;
+			TOnlineAsyncOpRef<FMergeableTestOpWithParams> Op = ServicesCommon.GetMergeableOp<FMergeableTestOpWithParams>(MoveTemp(Params1));
+			CHECK(!Op->IsReady());
+			Op->Then([Promise](FOnlineAsyncOp&) { Promise->SetValue(true); })
+				.Enqueue(ServicesCommon.GetParallelQueue());
+
+			FMergeableTestOpWithParams::Params Params2;
+			Params2.Index = 0;
+			TOnlineAsyncOpRef<FMergeableTestOpWithParams> Op2 = ServicesCommon.GetMergeableOp<FMergeableTestOpWithParams>(MoveTemp(Params2));
+			CHECK(Op2->IsReady());
+		});
+
+	RunToCompletion();
+}
+
+ASYNCOP_TEST_CASE("Verify that the cache returns different mergeable ops when the parameters are different")
+{
+	FTestPipeline& TestPipeline = GetPipeline();
+	TestPipeline.EmplaceAsyncLambda([](FAsyncLambdaResult Promise, SubsystemType Services)
+		{
+			UE::Online::FOnlineServicesCommon& ServicesCommon = static_cast<UE::Online::FOnlineServicesCommon&>(*Services);
+			FMergeableTestOpWithParams::Params Params1;
+			Params1.Index = 1;
+			TOnlineAsyncOpRef<FMergeableTestOpWithParams> Op = ServicesCommon.GetMergeableOp<FMergeableTestOpWithParams>(MoveTemp(Params1));
+			CHECK(!Op->IsReady());
+			Op->Then([](FOnlineAsyncOp&) {})
+				.Enqueue(ServicesCommon.GetParallelQueue());
+
+			FMergeableTestOpWithParams::Params Params2;
+			Params2.Index = 2;
+			TOnlineAsyncOpRef<FMergeableTestOpWithParams> Op2 = ServicesCommon.GetMergeableOp<FMergeableTestOpWithParams>(MoveTemp(Params2));
+			CHECK(!Op2->IsReady());
+			Op2->Then([Promise](FOnlineAsyncOp&) { Promise->SetValue(true); })
+				.Enqueue(ServicesCommon.GetParallelQueue());
+		});
+
+	RunToCompletion();
+}
+
+ASYNCOP_TEST_CASE("Verify that the various OnlineAsyncOp continuation forms complete (Immediately fulfilled promises)")
+{
+	FTestPipeline& TestPipeline = GetPipeline();
+	TestPipeline.EmplaceAsyncLambda([](FAsyncLambdaResult Promise, SubsystemType Services)
+		{
+			UE::Online::FOnlineServicesCommon& ServicesCommon = static_cast<UE::Online::FOnlineServicesCommon&>(*Services);
+			TOnlineAsyncOpRef<FTestOp> Op = ServicesCommon.GetOp<FTestOp>(FTestOp::Params());
+			Op->Then([](FOnlineAsyncOp&) {})
+				.Then([](FOnlineAsyncOp&) { return MakeFulfilledPromise<void>().GetFuture(); })
+				.Then([](FOnlineAsyncOp&) { return 1; } )
+				.Then([](FOnlineAsyncOp&, int) { return MakeFulfilledPromise<int>(1).GetFuture(); })
+				.Then([](FOnlineAsyncOp&, int, TPromise<TArray<int>>&& Promise) { Promise.SetValue(TArray<int>{1,2,3,4,5,6,7,8,9,0}); })
+				.ForEach([](FOnlineAsyncOp&, int) { return 1; })
+ 				.ForEach([](FOnlineAsyncOp&, int) { return MakeFulfilledPromise<int>(1).GetFuture(); })
+ 				.ForEach([](FOnlineAsyncOp&, int, TPromise<int>&& Promise) { Promise.SetValue(1); })
+ 				.ForEachN(3, [](FOnlineAsyncOp&, TArrayView<int> Value) { return TArray<int>(Value); })
+ 				.ForEachN(3, [](FOnlineAsyncOp&, TArrayView<int> Value) { return MakeFulfilledPromise<TArray<int>>(Value).GetFuture(); })
+ 				.ForEachN(3, [](FOnlineAsyncOp&, TArrayView<int> Value, TPromise<TArray<int>>&& Promise) { Promise.SetValue(TArray<int>(Value)); })
+				.Then([Promise](TOnlineAsyncOp<FTestOp>& Op, const TArray<int>&) { Op.SetResult(FTestOp::Result()); Promise->SetValue(true); })
+				.Enqueue(ServicesCommon.GetParallelQueue());
+		});
+	
+	RunToCompletion();
+}
+
+ASYNCOP_TEST_CASE("Verify that the various OnlineAsyncOp continuation forms complete (Promises not immediatly fulfilled)")
+{
+	FTestPipeline& TestPipeline = GetPipeline();
+	TestPipeline.EmplaceAsyncLambda([](FAsyncLambdaResult Promise, SubsystemType Services)
+		{
+			UE::Online::FOnlineServicesCommon& ServicesCommon = static_cast<UE::Online::FOnlineServicesCommon&>(*Services);
+			TOnlineAsyncOpRef<FTestOp> Op = ServicesCommon.GetOp<FTestOp>(FTestOp::Params());
+			Op->Then([](FOnlineAsyncOp&) {})
+				.Then([](FOnlineAsyncOp&) { return Async(EAsyncExecution::Thread, []() { return; }); })
+				.Then([](FOnlineAsyncOp&) { return 1; })
+				.Then([](FOnlineAsyncOp&, int) { return Async(EAsyncExecution::Thread, []() { return 1; }); })
+				.Then([](FOnlineAsyncOp&, int, TPromise<TArray<int>>&& Promise) { Async(EAsyncExecution::Thread, []() { return 1; }).Next([MovedPromise = MoveTemp(Promise)](int) mutable { MovedPromise.SetValue(TArray<int>{1, 2, 3, 4, 5, 6, 7, 8, 9, 0}); }); })
+				.ForEach([](FOnlineAsyncOp&, int) { return 1; })
+				.ForEach([](FOnlineAsyncOp&, int) { return Async(EAsyncExecution::Thread, []() { return 1; }); })
+				.ForEach([](FOnlineAsyncOp&, int, TPromise<int>&& Promise) { Async(EAsyncExecution::Thread, []() { return 1; }).Next([MovedPromise = MoveTemp(Promise)](int) mutable { MovedPromise.SetValue(1); }); })
+ 				.ForEachN(3, [](FOnlineAsyncOp&, TArrayView<int> Value) { return TArray<int>(Value); })
+ 				.ForEachN(3, [](FOnlineAsyncOp&, TArrayView<int> Value) { return Async(EAsyncExecution::Thread, [Result = TArray<int>(Value)]() { return Result; }); })
+				.ForEachN(3, [](FOnlineAsyncOp&, TArrayView<int> Value, TPromise<TArray<int>>&& Promise) { Async(EAsyncExecution::Thread, []() { return 1; }).Next([MovedPromise = MoveTemp(Promise), Result = TArray<int>(Value)](int) mutable { MovedPromise.SetValue(Result); }); })
+				.Then([Promise](TOnlineAsyncOp<FTestOp>& Op, const TArray<int>&) { Op.SetResult(FTestOp::Result()); Promise->SetValue(true); })
+				.Enqueue(ServicesCommon.GetParallelQueue());
+		});
+
+	RunToCompletion();
+}
+
+ASYNCOP_TEST_CASE("Verify that the various OnlineAsyncOp execution policies work as expected)")
+{
+	FTestPipeline& TestPipeline = GetPipeline();
+	uint32 ThreadId;
+	TestPipeline.EmplaceAsyncLambda([&ThreadId](FAsyncLambdaResult Promise, SubsystemType Services)
+		{
+			UE::Online::FOnlineServicesCommon& ServicesCommon = static_cast<UE::Online::FOnlineServicesCommon&>(*Services);
+			TOnlineAsyncOpRef<FTestOp> Op = ServicesCommon.GetOp<FTestOp>(FTestOp::Params());
+			Op->Then([&ThreadId](FOnlineAsyncOp&) {})
+				.Then([](FOnlineAsyncOp&)
+					{
+						// This is expected to be on the game thread
+						CHECK((FTaskGraphInterface::Get().IsCurrentThreadKnown() && FTaskGraphInterface::Get().GetCurrentThreadIfKnown() == ENamedThreads::GameThread));
+					} , FOnlineAsyncExecutionPolicy::RunOnNextTick())
+				.Then([&ThreadId](FOnlineAsyncOp&)
+					{
+						// Task graph won't know the thread pool thread
+						CHECK(!FTaskGraphInterface::Get().IsCurrentThreadKnown());
+						ThreadId = FPlatformTLS::GetCurrentThreadId();
+					}, FOnlineAsyncExecutionPolicy::RunOnThreadPool())
+				.Then([&ThreadId](FOnlineAsyncOp&)
+					{
+						// because the previous step returns immediately, this should be run on the same thread as the previous step
+						CHECK(ThreadId == FPlatformTLS::GetCurrentThreadId());
+					}, FOnlineAsyncExecutionPolicy::RunImmediately())
+				.Then([](FOnlineAsyncOp&)
+					{
+						// This can be on any task graph thread
+						CHECK(FTaskGraphInterface::Get().IsCurrentThreadKnown());
+					}, FOnlineAsyncExecutionPolicy::RunOnTaskGraph())
+				.Then([](FOnlineAsyncOp&)
+					{
+						// This is expected to be on the game thread
+						CHECK((FTaskGraphInterface::Get().IsCurrentThreadKnown() && FTaskGraphInterface::Get().GetCurrentThreadIfKnown() == ENamedThreads::GameThread));
+					}, FOnlineAsyncExecutionPolicy::RunOnGameThread())
+				.Then([Promise](TOnlineAsyncOp<FTestOp>& Op) { Op.SetResult(FTestOp::Result()); Promise->SetValue(true); })
+				.Enqueue(ServicesCommon.GetParallelQueue());
+		});
+
+	RunToCompletion();
+}
+
+ASYNCOP_TEST_CASE("Verify that the duration cache policy works)")
+{
+	FTestPipeline& TestPipeline = GetPipeline();
+	TArray<FString> CacheDurationTestHeiarchy = { TEXT("OnlineTests.TestCacheDuration") };
+	// Uses the following operation configuration in OnlineTests' DefaultEngine.ini
+	// [OnlineTests.TestCacheDuration]
+	// CacheExpiration=Duration
+	// CacheExpirySeconds=2
+	// bCacheError=false
+	TestPipeline.EmplaceLambda([&CacheDurationTestHeiarchy](SubsystemType Services)
+		{
+			UE::Online::FOnlineServicesCommon& ServicesCommon = static_cast<UE::Online::FOnlineServicesCommon&>(*Services);
+			TOnlineAsyncOpRef<FTestOp> Op = ServicesCommon.GetJoinableOp<FTestOp>(FTestOp::Params(), CacheDurationTestHeiarchy);
+			Op->Then([](FOnlineAsyncOp&) {})
+				.Then([](TOnlineAsyncOp<FTestOp>& Op) { Op.SetResult(FTestOp::Result()); })
+				.Enqueue(ServicesCommon.GetParallelQueue());
+		})
+		.EmplaceLambda([&CacheDurationTestHeiarchy](SubsystemType Services)
+		{
+			UE::Online::FOnlineServicesCommon& ServicesCommon = static_cast<UE::Online::FOnlineServicesCommon&>(*Services);
+			TOnlineAsyncOpRef<FTestOp> Op = ServicesCommon.GetJoinableOp<FTestOp>(FTestOp::Params(), CacheDurationTestHeiarchy);
+			CHECK(Op->IsReady()); // Cache should return the already in flight/completed op
+		})
+		.EmplaceAsyncLambda([](FAsyncLambdaResult Promise, SubsystemType Services)
+		{
+			// Wait a second into the cache duration
+			FTSTicker::GetCoreTicker().AddTicker(TEXT("OnlineTestTimeout"), 1.0f, [Promise](float) { Promise->SetValue(true); return false; });
+		})
+		.EmplaceLambda([&CacheDurationTestHeiarchy](SubsystemType Services)
+		{
+			UE::Online::FOnlineServicesCommon& ServicesCommon = static_cast<UE::Online::FOnlineServicesCommon&>(*Services);
+			TOnlineAsyncOpRef<FTestOp> Op = ServicesCommon.GetJoinableOp<FTestOp>(FTestOp::Params(), CacheDurationTestHeiarchy);
+			CHECK(Op->IsReady()); // Cache should return the already completed op
+		})
+		.EmplaceAsyncLambda([](FAsyncLambdaResult Promise, SubsystemType Services)
+		{
+			// Wait until shortly after the cache should expire
+			FTSTicker::GetCoreTicker().AddTicker(TEXT("OnlineTestTimeout"), 2.0f, [Promise](float) { Promise->SetValue(true); return false; });
+		})
+		.EmplaceAsyncLambda([&CacheDurationTestHeiarchy](FAsyncLambdaResult Promise, SubsystemType Services)
+		{
+			UE::Online::FOnlineServicesCommon& ServicesCommon = static_cast<UE::Online::FOnlineServicesCommon&>(*Services);
+			TOnlineAsyncOpRef<FTestOp> Op = ServicesCommon.GetJoinableOp<FTestOp>(FTestOp::Params(), CacheDurationTestHeiarchy);
+			CHECK(!Op->IsReady()); // Cache should return a new op
+			Op->Then([](FOnlineAsyncOp&) {})
+				.Then([Promise](TOnlineAsyncOp<FTestOp>& Op) { Op.SetResult(FTestOp::Result()); Promise->SetValue(true); })
+				.Enqueue(ServicesCommon.GetParallelQueue());
+		})
+		.EmplaceAsyncLambda([](FAsyncLambdaResult Promise, SubsystemType Services)
+		{
+			// Wait until shortly after the cache should expire
+			FTSTicker::GetCoreTicker().AddTicker(TEXT("OnlineTestTimeout"), 3.0f, [Promise](float) { Promise->SetValue(true); return false; });
+		});
+
+	RunToCompletion();
+}
+
+ASYNCOP_TEST_CASE("Verify that the cache error setting works)")
+{
+	DestroyCurrentServiceModule();
+	FTestPipeline& TestPipeline = GetPipeline();
+	TArray<FString> DontCacheErrorHeiarchy = { TEXT("OnlineTests.TestDontCacheError") };
+	// Uses the following operation configuration in OnlineTests' DefaultEngine.ini
+	// [OnlineTests.TestDontCacheError]
+	// CacheExpiration=Duration
+	// CacheExpirySeconds=1
+	// bCacheError=false
+	TArray<FString> CacheErrorHeiarchy = { TEXT("OnlineTests.TestCacheError") };
+	// Uses the following operation configuration in OnlineTests' DefaultEngine.ini
+	// [OnlineTests.TestCacheError]
+	// CacheExpiration=Duration
+	// CacheExpirySeconds=1
+	// bCacheError=true
+	TestPipeline.EmplaceAsyncLambda([&DontCacheErrorHeiarchy](FAsyncLambdaResult Promise, SubsystemType Services)
+		{
+			UE::Online::FOnlineServicesCommon& ServicesCommon = static_cast<UE::Online::FOnlineServicesCommon&>(*Services);
+			TOnlineAsyncOpRef<FTestOp> Op = ServicesCommon.GetJoinableOp<FTestOp>(FTestOp::Params(), DontCacheErrorHeiarchy);
+			Op->Then([](FOnlineAsyncOp&) {})
+				.Then([Promise](TOnlineAsyncOp<FTestOp>& Op) { Op.SetError(UE::Online::Errors::Unknown()); Promise->SetValue(true); })
+				.Enqueue(ServicesCommon.GetParallelQueue());
+		})
+		.EmplaceAsyncLambda([&CacheErrorHeiarchy](FAsyncLambdaResult Promise, SubsystemType Services)
+		{
+			UE::Online::FOnlineServicesCommon& ServicesCommon = static_cast<UE::Online::FOnlineServicesCommon&>(*Services);
+			TOnlineAsyncOpRef<FTestOp> Op = ServicesCommon.GetJoinableOp<FTestOp>(FTestOp::Params(), CacheErrorHeiarchy);
+			CHECK(!Op->IsReady()); // Should return a new op
+			Op->Then([](FOnlineAsyncOp&) {})
+				.Then([Promise](TOnlineAsyncOp<FTestOp>& Op) { Op.SetError(UE::Online::Errors::Unknown()); Promise->SetValue(true); })
+				.Enqueue(ServicesCommon.GetParallelQueue());
+		})
+		.EmplaceAsyncLambda([&CacheErrorHeiarchy](FAsyncLambdaResult Promise, SubsystemType Services)
+		{
+			UE::Online::FOnlineServicesCommon& ServicesCommon = static_cast<UE::Online::FOnlineServicesCommon&>(*Services);
+			TOnlineAsyncOpRef<FTestOp> Op = ServicesCommon.GetJoinableOp<FTestOp>(FTestOp::Params(), CacheErrorHeiarchy);
+			CHECK(Op->IsReady()); // Should return the cached, failed op
+			Op->GetHandle().OnComplete([Promise](const TOnlineResult<FTestOp>& Result)
+				{
+					CHECK(Result.IsError());
+					Promise->SetValue(true);
+				});
+		})
+		.EmplaceAsyncLambda([](FAsyncLambdaResult Promise, SubsystemType Services)
+		{
+			// Wait until shortly after the cache should expire
+			FTSTicker::GetCoreTicker().AddTicker(TEXT("OnlineTestTimeout"), 2.0f, [Promise](float) { Promise->SetValue(true); return false; });
+		});
+
+	RunToCompletion();
+}
+
+ASYNCOP_TEST_CASE("Verify that the cache expires once complete)")
+{
+	FTestPipeline& TestPipeline = GetPipeline();
+	TArray<FString> CacheExpireUponCompletionHeiarchy = { TEXT("OnlineTests.TestCacheExpireUponCompletion") };
+	// Uses the following operation configuration in OnlineTests' DefaultEngine.ini
+	// [OnlineTests.TestCacheExpireUponCompletion]
+	// CacheExpiration=UponCompletion
+	// bCacheError=false
+	TestPipeline.EmplaceAsyncLambda([&CacheExpireUponCompletionHeiarchy](FAsyncLambdaResult Promise, SubsystemType Services)
+		{
+			UE::Online::FOnlineServicesCommon& ServicesCommon = static_cast<UE::Online::FOnlineServicesCommon&>(*Services);
+			TOnlineAsyncOpRef<FTestOp> Op = ServicesCommon.GetJoinableOp<FTestOp>(FTestOp::Params(), CacheExpireUponCompletionHeiarchy);
+			Op->Then([](FOnlineAsyncOp&) {})
+				.Then([Promise](TOnlineAsyncOp<FTestOp>& Op) { Op.SetResult(FTestOp::Result()); Promise->SetValue(true); })
+				.Enqueue(ServicesCommon.GetParallelQueue());
+			TOnlineAsyncOpRef<FTestOp> Op2 = ServicesCommon.GetJoinableOp<FTestOp>(FTestOp::Params(), CacheExpireUponCompletionHeiarchy);
+			CHECK(Op2->IsReady());
+			CHECK(!Op2->IsComplete());
+		})
+		.EmplaceAsyncLambda([&CacheExpireUponCompletionHeiarchy](FAsyncLambdaResult Promise, SubsystemType Services)
+		{
+			UE::Online::FOnlineServicesCommon& ServicesCommon = static_cast<UE::Online::FOnlineServicesCommon&>(*Services);
+			TOnlineAsyncOpRef<FTestOp> Op = ServicesCommon.GetJoinableOp<FTestOp>(FTestOp::Params(), CacheExpireUponCompletionHeiarchy);
+			CHECK(!Op->IsReady());
+			Op->Then([](FOnlineAsyncOp&) {})
+				.Then([Promise](TOnlineAsyncOp<FTestOp>& Op) { Op.SetResult(FTestOp::Result()); Promise->SetValue(true); })
+				.Enqueue(ServicesCommon.GetParallelQueue());
+		});
+
+	RunToCompletion();
+}
diff --git a/Engine/Source/Programs/Online/OnlineTestsCore/Source/Private/Tests/TestErrorsAdapter.cpp b/Engine/Source/Programs/Online/OnlineTestsCore/Source/Private/Tests/TestErrorsAdapter.cpp
new file mode 100644
--- /dev/null
+++ b/Engine/Source/Programs/Online/OnlineTestsCore/Source/Private/Tests/TestErrorsAdapter.cpp
@@ -0,0 +1,70 @@
+// Copyright Epic Games, Inc. All Rights Reserved.
+#include "CoreMinimal.h"
+#include <catch2/catch_test_macros.hpp>
+
+#include "Online/OnlineError.h"
+#include "Online/OnlineErrorDefinitions.h"
+#include "Online/ErrorsOSSAdapter.h"
+#include "OnlineError.h"
+
+#define ADAPTER_ERROR_TAG "[AdapterErrors]"
+#define ADAPTER_ERROR_TEST_CASE(x, ...) TEST_CASE(x, ADAPTER_ERROR_TAG __VA_ARGS__)
+#define ONLINE_ERROR_NAMESPACE "errors.com.epicgames.testing"
+using FOnlineErrorOsvc = UE::Online::FOnlineError;
+
+ADAPTER_ERROR_TEST_CASE("Basic adapter error compilation test")
+{
+	FOnlineErrorOss TestErrorOss_Success = ONLINE_ERROR(EOnlineErrorResult::Success, TEXT("success"));
+	FOnlineErrorOss TestErrorOss_InvalidParams = ONLINE_ERROR(EOnlineErrorResult::InvalidParams, TEXT("invalid_params"));
+	FOnlineErrorOss TestErrorOss_Cancelled = ONLINE_ERROR(EOnlineErrorResult::Canceled, TEXT("custom_text_abc"));
+
+	using namespace UE::Online;
+	using namespace UE::Online::Errors;
+
+	FOnlineErrorOsvc TestError_Success = Errors::FromOssError(TestErrorOss_Success);
+	FOnlineErrorOsvc TestError_InvalidParams = Errors::FromOssError(TestErrorOss_InvalidParams);
+	FOnlineErrorOsvc TestError_Cancelled = Errors::FromOssError(TestErrorOss_Cancelled);
+
+	CHECK(TestError_Success == Errors::Success());
+	CHECK(TestError_Success == TestErrorOss_Success);
+	CHECK(TestError_Success != TestErrorOss_Cancelled);
+	CHECK(TestError_InvalidParams == Errors::InvalidParams());
+	CHECK(TestError_InvalidParams == TestErrorOss_InvalidParams);
+	CHECK(TestError_InvalidParams != TestError_Success);
+	CHECK(TestError_Cancelled == Errors::Cancelled());
+	CHECK(TestError_Cancelled == TestErrorOss_Cancelled);
+	CHECK(TestError_Cancelled != TestErrorOss_InvalidParams);
+	CHECK(TestError_Cancelled.GetLogString().Contains(TEXT("custom_text_abc")));
+}
+
+ADAPTER_ERROR_TEST_CASE("Check platform adapter is working properly")
+{
+
+	FOnlineErrorOss TestErrorOss_InvalidParams = ONLINE_ERROR(EOnlineErrorResult::InvalidParams, TEXT("WINDOWS_abcdef"));
+	FOnlineErrorOss TestErrorOss_Cancelled = ONLINE_ERROR(EOnlineErrorResult::Canceled, TEXT("abcdef"));
+
+	using namespace UE::Online;
+	using namespace UE::Online::Errors;
+
+	// this should obviously be done in a static initializer in final code, just doing it like this for quick tests
+	TOssPlatformErrorHandler ErrorHandler = [&](const FOnlineErrorOsvc& Error, const FOnlineErrorOss& Error2)
+	{
+		if (Error2.GetErrorCode().Contains(TEXT("WINDOWS")))
+		{
+			return TOptional<FOnlineErrorOsvc>(Errors::MissingInterface(Error));
+		}
+		return TOptional<FOnlineErrorOsvc>();
+	};
+	AddOssPlatformErrorHandler(MoveTemp(ErrorHandler));
+
+	FOnlineErrorOsvc TestError_InvalidParams = Errors::FromOssError(TestErrorOss_InvalidParams);
+	FOnlineErrorOsvc TestError_Cancelled = Errors::FromOssError(TestErrorOss_Cancelled);
+
+	CHECK(TestError_InvalidParams == Errors::MissingInterface());
+	CHECK(TestError_Cancelled != Errors::MissingInterface());
+
+	// If we make more than one test, we will need to empty out the platform functions array somehow, or else it will crash on a dangling delegate
+	//Private::PlatformFunctions.Empty();
+}
+
+#undef ONLINE_ERROR_NAMESPACE
\ No newline at end of file
diff --git a/Engine/Source/Programs/Online/OnlineTestsCore/Source/Public/Helpers/Auth/AuthLogin.h b/Engine/Source/Programs/Online/OnlineTestsCore/Source/Public/Helpers/Auth/AuthLogin.h
new file mode 100644
--- /dev/null
+++ b/Engine/Source/Programs/Online/OnlineTestsCore/Source/Public/Helpers/Auth/AuthLogin.h
@@ -0,0 +1,93 @@
+// Copyright Epic Games, Inc. All Rights Reserved.
+#pragma once
+#include "CoreMinimal.h"
+#include "TestDriver.h"
+#include "TestHarness.h"
+
+#include "Online/AuthCommon.h"
+#include "Online/OnlineAsyncOp.h"
+#include "Online/OnlineErrorDefinitions.h"
+#include "Online/OnlineServicesCommon.h"
+#include "OnlineCatchHelper.h"
+
+using namespace UE::Online;
+
+struct FAuthLoginStep : public FTestPipeline::FStep
+{
+	FAuthLoginStep(UE::Online::FAuthLogin::Params&& InLocalAccount)
+		: LocalAccount(MoveTemp(InLocalAccount))
+	{
+	}
+
+	virtual ~FAuthLoginStep()
+	{
+		if (OnlineAuthPtr != nullptr)
+		{
+			OnlineAuthPtr = nullptr;
+		}
+	};
+
+	enum class EState { Init, LoginCalled, Done } State = EState::Init;
+
+
+	virtual EContinuance Tick(SubsystemType OnlineSubsystem) override
+	{
+		OnlineAuthPtr = OnlineSubsystem->GetAuthInterface();
+
+		switch (State)
+		{
+		case EState::Init:
+		{
+			State = EState::LoginCalled;
+
+			FPlatformUserId PlatformUserId = LocalAccount.PlatformUserId;
+
+			TOnlineResult<FAuthGetLocalOnlineUserByPlatformUserId> LocalOnlineUserResult = OnlineAuthPtr->GetLocalOnlineUserByPlatformUserId({PlatformUserId});
+			if (LocalOnlineUserResult.IsOk() && LocalOnlineUserResult.GetOkValue().AccountInfo->LoginStatus == UE::Online::ELoginStatus::LoggedIn)
+			{
+				State = EState::Done;
+			}
+			else
+			{
+				OnlineAuthPtr->Login(MoveTemp(LocalAccount))
+				 .OnComplete([this, PlatformUserId](const TOnlineResult<FAuthLogin> Op) mutable
+				{
+					CHECK(State == EState::LoginCalled);
+					CHECK_OP_EQ(Op, Errors::NotImplemented());
+
+					if(Op.IsOk())
+					{
+						CHECK_OP(OnlineAuthPtr->GetLocalOnlineUserByOnlineAccountId({Op.GetOkValue().AccountInfo->AccountId}));
+					}
+					else if (Op.GetErrorValue() == Errors::NotImplemented())
+					{
+						// Some auth implementations do not have an explicit login / logout. In those implementations all platform users are assumed to always be logged in.
+						TOnlineResult<FAuthGetLocalOnlineUserByPlatformUserId> OnlineUserResult = OnlineAuthPtr->GetLocalOnlineUserByPlatformUserId({PlatformUserId});
+						CHECK_OP(OnlineUserResult);
+					}
+
+					State = EState::Done;
+				});
+			}
+
+			break;
+		}
+		case EState::LoginCalled:
+		{
+			break;
+		}
+		case EState::Done:
+		{
+			return EContinuance::Done;
+		}
+		}
+
+		return EContinuance::ContinueStepping;
+	}
+
+protected:
+	int32 LocalUserNum = 0;
+	UE::Online::FAuthLogin::Params LocalAccount;
+	UE::Online::IAuthPtr OnlineAuthPtr = nullptr;
+	FDelegateHandle	OnLoginCompleteDelegateHandle;
+};
\ No newline at end of file
diff --git a/Engine/Source/Programs/Online/OnlineTestsCore/Source/Public/Helpers/Auth/AuthLogout.h b/Engine/Source/Programs/Online/OnlineTestsCore/Source/Public/Helpers/Auth/AuthLogout.h
new file mode 100644
--- /dev/null
+++ b/Engine/Source/Programs/Online/OnlineTestsCore/Source/Public/Helpers/Auth/AuthLogout.h
@@ -0,0 +1,75 @@
+// Copyright Epic Games, Inc. All Rights Reserved.
+#pragma once
+
+#include "CoreMinimal.h"
+#include "TestDriver.h"
+#include "TestHarness.h"
+
+#include "Online/AuthCommon.h"
+#include "Online/OnlineAsyncOp.h"
+#include "Online/OnlineErrorDefinitions.h"
+#include "Online/OnlineServicesCommon.h"
+
+struct FAuthLogoutStep : public FTestPipeline::FStep
+{
+	FAuthLogoutStep(FPlatformUserId InPlatformUserId)
+		: PlatformUserId(InPlatformUserId)
+	{}
+
+	virtual ~FAuthLogoutStep()
+	{
+		OnlineAuthPtr = nullptr;
+	};
+
+
+	enum class EState { Init, LogoutCalled, Done } State = EState::Init;
+
+
+	virtual EContinuance Tick(SubsystemType OnlineSubsystem) override
+	{
+		OnlineAuthPtr = OnlineSubsystem->GetAuthInterface();
+
+		switch (State)
+		{
+		case EState::Init:
+		{
+			State = EState::LogoutCalled;
+			TOnlineResult<FAuthGetLocalOnlineUserByPlatformUserId> AccountId = OnlineAuthPtr->GetLocalOnlineUserByPlatformUserId({PlatformUserId});
+			CAPTURE(ToLogString(AccountId), PlatformUserId);
+			CHECK_OP(AccountId);
+			if(AccountId.IsOk())
+			{
+				OnlineAuthPtr->Logout({ AccountId.GetOkValue().AccountInfo->AccountId})
+					.OnComplete([this](const TOnlineResult<FAuthLogout> Op) mutable
+					{
+						CHECK(State == EState::LogoutCalled);
+						// Some implementations do not implement an explicit login / logout.
+						CHECK_OP_EQ(Op, Errors::NotImplemented());
+
+						State = EState::Done;
+					});
+			}
+			else
+			{
+				State = EState::Done;
+			}
+
+			break;
+		}
+		case EState::LogoutCalled:
+		{
+			break;
+		}
+		case EState::Done:
+		{
+			return EContinuance::Done;
+		}
+		}
+
+		return EContinuance::ContinueStepping;
+	}
+protected:
+	FPlatformUserId PlatformUserId;
+	UE::Online::IAuthPtr OnlineAuthPtr = nullptr;
+	FDelegateHandle	OnLogoutCompleteDelegateHandle;
+};
\ No newline at end of file
diff --git a/Engine/Source/Programs/Online/OnlineTestsCore/Source/Public/Helpers/LambdaStep.h b/Engine/Source/Programs/Online/OnlineTestsCore/Source/Public/Helpers/LambdaStep.h
new file mode 100644
--- /dev/null
+++ b/Engine/Source/Programs/Online/OnlineTestsCore/Source/Public/Helpers/LambdaStep.h
@@ -0,0 +1,23 @@
+// Copyright Epic Games, Inc. All Rights Reserved.
+#pragma once
+#include "TestDriver.h"
+#include "TestHarness.h"
+
+struct FLambdaStep : public FTestPipeline::FStep
+{
+	FLambdaStep(TFunction<void(SubsystemType)>&& InFunction)
+		: Function(MoveTemp(InFunction))
+	{
+	}
+
+	virtual ~FLambdaStep() = default;
+
+	virtual EContinuance Tick(SubsystemType OnlineSubsystem) override
+	{
+		Function(OnlineSubsystem);
+		return EContinuance::Done;
+	}
+
+protected:
+	TFunction<void(SubsystemType)> Function;
+};
\ No newline at end of file
diff --git a/Engine/Source/Programs/Online/OnlineTestsCore/Source/Public/Helpers/Stats/UpdateStatsHelper.h b/Engine/Source/Programs/Online/OnlineTestsCore/Source/Public/Helpers/Stats/UpdateStatsHelper.h
new file mode 100644
--- /dev/null
+++ b/Engine/Source/Programs/Online/OnlineTestsCore/Source/Public/Helpers/Stats/UpdateStatsHelper.h
@@ -0,0 +1,19 @@
+// Copyright Epic Games, Inc. All Rights Reserved.
+#pragma once
+
+#include "Online/Stats.h"
+#include "OnlineCatchHelper.h"
+
+namespace UE::Online
+{
+
+inline void UpdateStats_Fixture(FAsyncLambdaResult Promise, SubsystemType Services, const FAccountId& AccountId,  const TArray<FUserStats>& UpdateUsersStats)
+{
+	Services->GetStatsInterface()->UpdateStats({AccountId,  UpdateUsersStats}).OnComplete([Promise=MoveTemp(Promise)](const TOnlineResult<UE::Online::FUpdateStats>& Result) mutable
+	{
+		CHECK_OP(Result);
+		Promise->SetValue(true);
+	});
+}
+
+}
diff --git a/Engine/Source/Programs/Online/OnlineTestsCore/Source/Public/OnlineCatchHelper.h b/Engine/Source/Programs/Online/OnlineTestsCore/Source/Public/OnlineCatchHelper.h
new file mode 100644
--- /dev/null
+++ b/Engine/Source/Programs/Online/OnlineTestsCore/Source/Public/OnlineCatchHelper.h
@@ -0,0 +1,224 @@
+// Copyright Epic Games, Inc. All Rights Reserved.
+
+#pragma once
+
+#include "TestHarness.h"
+#include "TestDriver.h"
+#include "CoreMinimal.h"
+#include "OnlineCatchStringMakers.h"
+#include "Interfaces/OnlineIdentityInterface.h"
+#include "Online/AuthCommon.h"
+#include "Online/OnlineAsyncOp.h"
+#include "Online/OnlineError.h"
+#include "Online/OnlineErrorDefinitions.h"
+#include "Online/OnlineUtils.h"
+#include "Online/OnlineResult.h"
+#include "Online/OnlineServicesCommon.h"
+
+#include <catch2/catch_test_case_info.hpp>
+#include <catch2/catch_test_macros.hpp>
+#include <catch2/interfaces/catch_interfaces_registry_hub.hpp>
+#include <catch2/catch_active_test.hpp>
+#include <string>
+using namespace UE::Online;
+
+class OnlineTestBase : public Catch::ITestInvoker
+{
+public:
+	void ConstructInternal(FString ServiceName, UE::Online::EOnlineServices ServiceType);
+
+	// must be public so that Catch2::Detail::unique_ptr can call the destructor
+	virtual ~OnlineTestBase();
+
+	/* Loads all necessary services for the current test run */
+	static void LoadServiceModules();
+
+	/* Unloads all necessary services for the current test run */
+	static void UnloadServiceModules();
+
+	FTestPipeline& GetLoginPipeline(uint32 NumUsersToLogin = 1) const;
+
+
+	typedef UE::Online::FAccountId FAccountId;
+public:
+	/* Creates a login pipeline and logs in user as many users as FAccountIds passed in, immediately assigns FAccountIds to the passed in arguments.
+		This function relies on FAuthLoginStep to be executed by RunToCompletion(...) so will flush any queued steps in the test. */
+	FTestPipeline& GetLoginPipelineArray(std::initializer_list<std::reference_wrapper<FAccountId>> AccountIdsArr) const;
+	/* Creates a login pipeline and logs in user, immediately assigns FAccountIds to the passed in arguments for one account */
+	FTestPipeline& GetLoginPipeline(FAccountId& OutAccountId) const;
+	/* Creates a login pipeline and logs in user, immediately assigns FAccountIds to the passed in arguments for two accounts */
+	FTestPipeline& GetLoginPipeline(FAccountId& OutAccountId, FAccountId& OutAccountId2 ) const;
+	/* Creates a login pipeline and logs in user, immediately assigns FAccountIds to the passed in arguments for three accounts */
+	FTestPipeline& GetLoginPipeline(FAccountId& OutAccountId, FAccountId& OutAccountId2, FAccountId& OutAccountId3) const;
+	/* Creates a login pipeline and logs in user, immediately assigns FAccountIds to the passed in arguments for four accounts */
+	FTestPipeline& GetLoginPipeline(FAccountId& OutAccountId, FAccountId& OutAccountId2, FAccountId& OutAccountId3, FAccountId& OutAccountId4) const;
+	/* Creates a login pipeline and logs in user, immediately assigns FAccountIds to the passed in arguments for five accounts */
+	FTestPipeline& GetLoginPipeline(FAccountId& OutAccountId, FAccountId& OutAccountId2, FAccountId& OutAccountId3, FAccountId& OutAccountId4, FAccountId& OutAccountId5) const;
+
+	/* Attempts to assgin OutAccountId to what LocalUserId is logged in as. */
+	void AssignLoginUsers(int32 LocalUserId, FAccountId& OutAccountId) const;
+
+	FTestPipeline& GetPipeline() const;
+protected:
+
+	OnlineTestBase();
+
+	/* Helper function to delete current accounts on the TestDataService and wait for it to finish */
+	bool DeleteAccountsForCurrentTemplate() const;
+
+	/* Helper function to destroy the current OnlineService module which stores some state that maybe needs to be reset */
+	void DestroyCurrentServiceModule() const;
+
+	/* Proxy function to the DestroyCurrentServiceModule() */
+	bool ResetAccountStatus() const;
+
+	/* Returns the name of the service we're currently testing */
+	FString GetService() const;
+	UE::Online::EOnlineServices GetServiceType() const;
+	SubsystemType GetSubsystem() const;
+
+// If you define a custom credentials path you must define these in your own Auth.cpp file
+#if ONLINETESTS_USEEXTERNAUTH
+	// This will provide a paramters to pass to the OSS Auth interface.
+	UE::Online::FAuthLogin::Params CustomCredentials(int LocalUserNum) const;
+
+	// When called should "reset" accounts to default state for Stats, Achivements, Lobbies and other
+	// stateful interfaces to work.
+	bool CustomResetAccounts() const;
+
+	// When called will *remove* the account from the underlying service.
+	// This is generally called to test new-user login or first-time setup.
+	// To be called rare-ly.
+	bool CustomDeleteAccounts() const;
+#endif
+
+	UE::Online::FAuthLogin::Params GetIniCredentials(int LocalUserNum) const;
+	UE::Online::FAuthLogin::Params GetCredentials(int LocalUserNum) const;
+
+	/* Returns the ini login category name for the configured service */
+	FString GetLoginCredentialCategory() const;
+
+	void RunToCompletion(bool bLogout = true) const;
+
+	/* ITestInvoker */
+	virtual void invoke() const override = 0;
+
+private:
+	FString Tags;
+	FString Service;
+	UE::Online::EOnlineServices ServiceType;
+	// Catch's ITestInvoker is a const interface but we'll be changing stuff (emplacing steps into the driver, setting flags, etc.) so we make these mutable
+	mutable FTestDriver Driver;
+	mutable TSharedPtr<FTestPipeline> Pipeline;
+	mutable uint32 NumLocalUsers = -1;
+};
+
+typedef OnlineTestBase* (*OnlineTestConstructor)();
+
+class OnlineAutoReg
+{
+public:
+	struct FReportingSkippableTags
+	{
+		TArray<FString> MayFailTags;
+		TArray<FString> ShouldFailTags;
+		TArray<FString> DisableTestTags;
+	};
+
+	/*
+	* Helper function that calls CheckAllTagsIsIn(const TArray<FString>&, const TArray<FString>&);
+	* 
+	* @param  TestTags			The array of test tags we wish to test against.
+	* @param  RawTagString		Comma seperated string of elemnets we wish to convert to an array.
+	* 
+	* @return true if all elements of RawTagString prased as an comma sperated array is in TestTags.
+	*/
+	static bool CheckAllTagsIsIn(const TArray<FString>& TestTags, const FString& RawTagString);
+	
+
+	/**
+	 * Checks if every element of InputTags is in TestTags.
+	 *
+	 * @param  TestTags       The array of test tags we wish to test against.
+	 * @param  InputTags	  The array of tags we wish to see if all elements of are present in TestTags.
+	 *
+	 * @return  true if all elements of InputTags is in TestTags.
+	 */
+	static bool CheckAllTagsIsIn(const TArray<FString>& TestTags, const TArray<FString>& InputTags);
+	
+	static FString GenerateTags(const FString& ServiceName, const FReportingSkippableTags& SkippableTags, const TCHAR* InTag);
+	static bool ShouldDisableTest(const FString& ServiceName, const FReportingSkippableTags& SkippableTags, const TCHAR* InTag);
+	static bool ShouldSkipTest(const FString& TagsToCheck);
+	static void CheckRunningTestSkipOnTags();
+
+	// This code is kept identical to Catch internals so that there is as little deviation from OSS_TESTS and Online_OSS_TESTS as possible
+	OnlineAutoReg(OnlineTestConstructor TestCtor, Catch::SourceLineInfo LineInfo, const char* Name, const char* Tags, const char* AddlOnlineInfo);
+
+	struct FApplicableServicesConfig
+	{
+		FString Tag;
+		EOnlineServices ServicesType;
+		TArray<FString> ModulesToLoad;
+	};
+
+	static TArray<FApplicableServicesConfig> GetApplicableServices();
+};
+
+TArray<TFunction<void()>>* GetGlobalInitalizers();
+
+// INTERNAL_CATCH_UNIQUE_NAME is what Catch uses to translate the name of the test into a unique identifier without any spaces
+
+#define INTERNAL_ONLINE_TEST_CASE_NAMED(RegName, Name, Tags, ...)\
+namespace {\
+class PREPROCESSOR_JOIN(OnlineTest_,RegName) : public OnlineTestBase {\
+protected:\
+	virtual void invoke() const override;\
+};\
+	OnlineTestBase* PREPROCESSOR_JOIN(Construct_,RegName) (){\
+		return new PREPROCESSOR_JOIN(OnlineTest_,RegName) ();\
+	}\
+	OnlineAutoReg RegName( PREPROCESSOR_JOIN(&Construct_,RegName) , CATCH_INTERNAL_LINEINFO , Name, Tags, "");\
+}\
+void PREPROCESSOR_JOIN(OnlineTest_,RegName)::invoke() const\
+
+#define ONLINE_TEST_CASE(Name, Tags, ...) \
+	INTERNAL_ONLINE_TEST_CASE_NAMED(INTERNAL_CATCH_UNIQUE_NAME(OnlineRegistrar), Name, Tags, __VA_ARGS__)
+
+// Helpers for checking+capturing at the same time
+
+#define CHECK_OP(Op)\
+	CAPTURE(Op);\
+	CHECK(Op.IsOk());
+
+// also passes if the op is a success
+#define CHECK_OP_EQ(Op, Arg)\
+	CAPTURE(Op);\
+	CHECK((Op.IsOk() || Op.GetErrorValue() == Arg));
+
+
+#define CHECK_STR(Str)\
+	CAPTURE(Str);\
+	CHECK(!Str.IsEmpty());
+
+#define CHECK_STR_EMPTY(Str)\
+	CAPTURE(Str);\
+	CHECK(Str.IsEmpty());
+
+
+#define REQUIRE_OP(Op)\
+	CAPTURE(Op);\
+	REQUIRE(Op.IsOk());
+
+// also passes if the op is a success
+#define REQUIRE_OP_EQ(Op, Arg)\
+	CAPTURE(Op);\
+	REQUIRE((Op.IsOk() || Op.GetErrorValue() == Arg));
+
+
+#define REQUIRE_STR(Str)\
+	CAPTURE(Str);\
+	REQUIRE(!Str.IsEmpty());
+
+#define REQUIRE_STR_EMPTY(Str)\
+	CAPTURE(Str);\
+	REQUIRE(Str.IsEmpty());
\ No newline at end of file
