OnlineTestCore: Fix issue with UE::Error being ambiguous due to depenency move in TestEOSError and TestError
Fix warning SESSION_TEST_CASE macro being duplicately defined.

diff --git a/Engine/Source/Programs/Online/OnlineTestsCore/Source/OnlineTestsCore.Build.cs b/Engine/Source/Programs/Online/OnlineTestsCore/Source/OnlineTestsCore.Build.cs
--- a/Engine/Source/Programs/Online/OnlineTestsCore/Source/OnlineTestsCore.Build.cs
+++ b/Engine/Source/Programs/Online/OnlineTestsCore/Source/OnlineTestsCore.Build.cs
@@ -26,26 +26,20 @@ public class OnlineTestsCore : ModuleRules
 				"CoreUObject",
 				"Projects",
 				"EngineSettings",
+				"EOSSDK",
+				"EOSShared",
 				"OnlineSubsystem",
 				"OnlineServicesInterface",
 				"OnlineServicesCommon",
 				"OnlineServicesEOS",
 				"OnlineServicesNull",
 				"OnlineServicesOSSAdapter",
-				"EOSSDK",
-				"EOSShared",
 				"SSL",
 				"Json",
 				"JsonUtilities"
 			}
 		);
 
-		PublicDependencyModuleNames.AddRange(
-		new string[]
-		{
-			"OnlineSubsystem"
-		});
-
 		PublicDefinitions.Add(String.Format("ONLINETESTS_REQUIREAPPLICATIONTICK={0}", bRequireApplicationTick ? 1 : 0));
 		PublicDefinitions.Add(String.Format("ONLINETESTS_REQUIREPLATFORMINIT={0}", bRequirePlatformInit ? 1 : 0));
 		PublicDefinitions.Add(String.Format("ONLINETESTS_PLATFORMFILENAME={0}", PlatformFileName));
diff --git a/Engine/Source/Programs/Online/OnlineTestsCore/Source/Private/Tests/EOS/TestErrorsEOS.cpp b/Engine/Source/Programs/Online/OnlineTestsCore/Source/Private/Tests/EOS/TestErrorsEOS.cpp
deleted file mode 100644
diff --git a/Engine/Source/Programs/Online/OnlineTestsCore/Source/Private/Tests/Sessions/UpdateSessionSettingsTests.cpp b/Engine/Source/Programs/Online/OnlineTestsCore/Source/Private/Tests/Sessions/UpdateSessionSettingsTests.cpp
--- a/Engine/Source/Programs/Online/OnlineTestsCore/Source/Private/Tests/Sessions/UpdateSessionSettingsTests.cpp
+++ b/Engine/Source/Programs/Online/OnlineTestsCore/Source/Private/Tests/Sessions/UpdateSessionSettingsTests.cpp
@@ -7,9 +7,9 @@
 
 #define SESSIONS_TAGS "[suite_sessions]"
 #define SESSIONS_UPDATESESSIONSETTINGS_TAG SESSIONS_TAGS "[updatesessionsettings]"
-#define SESSIONS_TEST_CASE(x, ...) ONLINE_TEST_CASE(x, SESSIONS_TAGS __VA_ARGS__)
+#define UPDATESESSIONS_TEST_CASE(x, ...) ONLINE_TEST_CASE(x, SESSIONS_TAGS __VA_ARGS__)
 
-SESSIONS_TEST_CASE("If I call UpdateSessionSettings with an invalid account id, I get an error", SESSIONS_UPDATESESSIONSETTINGS_TAG)
+UPDATESESSIONS_TEST_CASE("If I call UpdateSessionSettings with an invalid account id, I get an error", SESSIONS_UPDATESESSIONSETTINGS_TAG)
 {
 	FUpdateSessionSettings::Params OpUpdateParams;
 	FUpdateSessionSettingsHelper::FHelperParams UpdateSessionSettingsHelperParams;
@@ -23,7 +23,7 @@ SESSIONS_TEST_CASE("If I call UpdateSessionSettings with an invalid account id,
 	RunToCompletion();
 }
 
-SESSIONS_TEST_CASE("If I call UpdateSessionSettings with an empty session name, I get an error", SESSIONS_UPDATESESSIONSETTINGS_TAG)
+UPDATESESSIONS_TEST_CASE("If I call UpdateSessionSettings with an empty session name, I get an error", SESSIONS_UPDATESESSIONSETTINGS_TAG)
 {
 	FAccountId AccountId;
 
@@ -43,7 +43,7 @@ SESSIONS_TEST_CASE("If I call UpdateSessionSettings with an empty session name,
 	RunToCompletion();
 }
 
-SESSIONS_TEST_CASE("If I call UpdateSessionSettings with an empty schema name in settings, I get an error", SESSIONS_UPDATESESSIONSETTINGS_TAG)
+UPDATESESSIONS_TEST_CASE("If I call UpdateSessionSettings with an empty schema name in settings, I get an error", SESSIONS_UPDATESESSIONSETTINGS_TAG)
 {
 	FAccountId AccountId;
 
@@ -64,7 +64,7 @@ SESSIONS_TEST_CASE("If I call UpdateSessionSettings with an empty schema name in
 	RunToCompletion();
 }
 
-SESSIONS_TEST_CASE("If I call UpdateSessionSettings with an invalid max connections number in settings, I get an error", SESSIONS_UPDATESESSIONSETTINGS_TAG)
+UPDATESESSIONS_TEST_CASE("If I call UpdateSessionSettings with an invalid max connections number in settings, I get an error", SESSIONS_UPDATESESSIONSETTINGS_TAG)
 {
 	FAccountId AccountId;
 
@@ -85,7 +85,7 @@ SESSIONS_TEST_CASE("If I call UpdateSessionSettings with an invalid max connecti
 	RunToCompletion();
 }
 
-SESSIONS_TEST_CASE("If I call UpdateSessionSettings with an empty custom setting name in settings, I get an error", SESSIONS_UPDATESESSIONSETTINGS_TAG)
+UPDATESESSIONS_TEST_CASE("If I call UpdateSessionSettings with an empty custom setting name in settings, I get an error", SESSIONS_UPDATESESSIONSETTINGS_TAG)
 {
 	FAccountId AccountId;
 
@@ -106,7 +106,7 @@ SESSIONS_TEST_CASE("If I call UpdateSessionSettings with an empty custom setting
 	RunToCompletion();
 }
 
-SESSIONS_TEST_CASE("If I call UpdateSessionSettings with an unregistered session name, I get an error", SESSIONS_UPDATESESSIONSETTINGS_TAG)
+UPDATESESSIONS_TEST_CASE("If I call UpdateSessionSettings with an unregistered session name, I get an error", SESSIONS_UPDATESESSIONSETTINGS_TAG)
 {
 	FAccountId AccountId;
 
@@ -126,7 +126,7 @@ SESSIONS_TEST_CASE("If I call UpdateSessionSettings with an unregistered session
 	RunToCompletion();
 }
 
-SESSIONS_TEST_CASE("If I call UpdateSessionSettings with valid data, the operation completes successfully", SESSIONS_UPDATESESSIONSETTINGS_TAG)
+UPDATESESSIONS_TEST_CASE("If I call UpdateSessionSettings with valid data, the operation completes successfully", SESSIONS_UPDATESESSIONSETTINGS_TAG)
 {
 	DestroyCurrentServiceModule();
 	ResetAccountStatus();
diff --git a/Engine/Source/Programs/Online/OnlineTestsCore/Source/Private/Tests/TestErrorsAdapter.cpp b/Engine/Source/Programs/Online/OnlineTestsCore/Source/Private/Tests/TestErrorsAdapter.cpp
deleted file mode 100644
