OSSTests - Enable OnlineSubsystemMcp testing

- Splitting OSSTests into two parts: OSSTests and OSSTestsCore.
- This division allowed to make OSMcpTests, which includes Mcp support for tests.
- OSSTests and OSMcpTests have their own DefaultEngine.ini from which they can get the necessary values for tests.

diff --git a/Engine/Source/Programs/Online/OSSTestsCore/OSSTestsCore.Build.cs b/Engine/Source/Programs/Online/OSSTestsCore/OSSTestsCore.Build.cs
new file mode 100644
--- /dev/null
+++ b/Engine/Source/Programs/Online/OSSTestsCore/OSSTestsCore.Build.cs
@@ -0,0 +1,32 @@
+// Copyright Epic Games, Inc. All Rights Reserved.
+
+using UnrealBuildTool;
+using System.IO;
+
+public class OSSTestsCore : ModuleRules
+{
+	public OSSTestsCore(ReadOnlyTargetRules Target) : base(Target)
+	{
+		PublicIncludePaths.Add("Programs/Online/OSSTestsCore");
+
+		PrivateDependencyModuleNames.AddRange(
+			new string[]
+			{
+				"Core",
+				"CoreUObject",
+				"Projects",
+				"EngineSettings",
+				"OnlineSubsystem",
+				"OnlineSubsystemNull",
+				"ApplicationCore"
+			}
+		);
+
+		PrivateIncludePaths.AddRange(
+			new string[] {
+				Path.Combine(Target.RelativeEnginePath, "Plugins", "Online"),
+				Path.Combine(Target.RelativeEnginePath, "Restricted", "NotForLicensees", "Plugins", "Online")
+            }
+        );
+	}
+}
diff --git a/Engine/Source/Programs/Online/OSSTestsCore/OnlineSubsystemCatchHelper.cpp b/Engine/Source/Programs/Online/OSSTestsCore/OnlineSubsystemCatchHelper.cpp
new file mode 100644
--- /dev/null
+++ b/Engine/Source/Programs/Online/OSSTestsCore/OnlineSubsystemCatchHelper.cpp
@@ -0,0 +1,274 @@
+// Copyright Epic Games, Inc. All Rights Reserved.
+
+#include "OnlineSubsystemCatchHelper.h"
+
+#include "Algo/AllOf.h"
+#include "Algo/Sort.h"
+#include "Algo/ForEach.h"
+#include "Interfaces/OnlineIdentityInterface.h"
+#include "Helpers/Identity/IdentityAutoLoginHelper.h"
+#include "Helpers/Identity/IdentityLoginHelper.h"
+#include "Helpers/Identity/IdentityLogoutHelper.h"
+#include "OnlineSubsystemNames.h"
+#include "Misc/CommandLine.h"
+#include "Misc/ConfigCacheIni.h"
+
+TArray<TFunction<void()>>* GetGlobalInitalizers()
+{
+	static TArray<TFunction<void()>> gInitalizersToCallInMain;
+	return &gInitalizersToCallInMain;
+}
+
+void OnlineSubsystemTestBase::ConstructInternal(FString SubsystemName)
+{
+	Subsystem = SubsystemName;
+}
+
+OnlineSubsystemTestBase::OnlineSubsystemTestBase()
+	: Driver()
+	, Pipeline(Driver.MakePipeline())
+{
+	// handle most cxn in ConstructInternal
+}
+
+OnlineSubsystemTestBase::~OnlineSubsystemTestBase()
+{
+
+}
+
+FString OnlineSubsystemTestBase::GetSubsystem() const
+{
+	return Subsystem;
+}
+
+FOnlineAccountCredentials OnlineSubsystemTestBase::GetCredentials(int LocalUserNum) const
+{
+	FString LoginCredentialCategory = FString::Printf(TEXT("LoginCredentials %s"), *Subsystem);
+	TArray<FString> CredentialsArr;
+	GConfig->GetArray(*LoginCredentialCategory, TEXT("Credentials"), CredentialsArr, GEngineIni);
+
+	if (LocalUserNum > CredentialsArr.Num())
+	{
+		UE_LOG(LogOSSTests, Error, TEXT("Attempted to GetCredentials for more than we have stored! Add more credentials to the DefaultEngine.ini for OssTests"));
+		return FOnlineAccountCredentials(FString(TEXT("")), FString(TEXT("")), FString(TEXT("")));
+	}
+	
+	FString LoginUsername, LoginType, LoginPassword;
+	FParse::Value(*CredentialsArr[LocalUserNum], TEXT("Type="), LoginType);
+	FParse::Value(*CredentialsArr[LocalUserNum], TEXT("Id="), LoginUsername);
+	FParse::Value(*CredentialsArr[LocalUserNum], TEXT("Password="), LoginPassword);
+
+	// Return an account, using the default first user if we only have one account for the subsystem.
+	return FOnlineAccountCredentials(LoginType, LoginUsername, LoginPassword);
+}
+
+FTestPipeline& OnlineSubsystemTestBase::GetLoginPipeline(uint32 NumUsersToLogin) const
+{
+	REQUIRE(NumLocalUsers == -1); // Don't call GetLoginPipeline more than once per test
+	NumLocalUsers = NumUsersToLogin;
+
+	bool bUseAutoLogin;
+	FString LoginCredentialCategory = FString::Printf(TEXT("LoginCredentials %s"), *Subsystem);
+	GConfig->GetBool(*LoginCredentialCategory, TEXT("UseAutoLogin"), bUseAutoLogin, GEngineIni);
+
+	if (bUseAutoLogin)
+	{
+		NumLocalUsers = 1;
+		Pipeline.EmplaceStep<FIdentityAutoLoginStep>(0);
+	}
+	else
+	{
+		for (uint32 i = 0; i < NumUsersToLogin; i++)
+		{
+			Pipeline.EmplaceStep<FIdentityLoginStep>(i, GetCredentials(i));
+		}
+	}
+
+	return Pipeline;
+}
+
+FTestPipeline& OnlineSubsystemTestBase::GetPipeline()
+{
+	return GetLoginPipeline(0);
+}
+
+void OnlineSubsystemTestBase::RunToCompletion() const
+{
+	for (uint32 i = 0; i < NumLocalUsers; i++)
+	{
+		Pipeline.EmplaceStep<FIdentityLogoutStep>(i);
+	}
+
+	FName SubsystemName = FName(GetSubsystem());
+	CAPTURE(*GetSubsystem());
+	FPipelineTestContext TestContext = FPipelineTestContext(SubsystemName);
+	CHECK(Driver.AddPipeline(MoveTemp(Pipeline), TestContext));
+	REQUIRE(IOnlineSubsystem::IsEnabled(SubsystemName));
+	Driver.RunToCompletion();
+}
+
+// TODO: This should poll some info from the tags and generate the list based on that
+TArray<FString> OnlineSubsystemAutoReg::GetApplicableSubsystems()
+{
+	TArray<FString> Subsystems;
+	GConfig->GetArray(TEXT("OnlineSubsystemTests"), TEXT("Subsystems"), Subsystems, GEngineIni);
+	return Subsystems;
+}
+
+bool OnlineSubsystemAutoReg::CheckAllTagsIsIn(const TArray<FString>& TestTags, const TArray<FString>& InputTags)
+{
+	if (InputTags.Num() == 0)
+	{
+		return false;
+	}
+
+	if (InputTags.Num() > TestTags.Num())
+	{
+		return false;
+	}
+
+	bool bAllInputTagsInTestTags = Algo::AllOf(InputTags, [&TestTags](const FString& CheckTag) -> bool
+		{
+			auto CheckStringCaseInsenstive = [&CheckTag](const FString& TestString) -> bool
+			{
+				return TestString.Equals(CheckTag, ESearchCase::IgnoreCase);
+			};
+
+			if (TestTags.ContainsByPredicate(CheckStringCaseInsenstive))
+			{
+				return true;
+			}
+
+			return false;
+		});
+
+	return bAllInputTagsInTestTags;
+}
+
+bool OnlineSubsystemAutoReg::CheckAllTagsIsIn(const TArray<FString>& TestTags, const FString& RawTagString)
+{
+	TArray<FString> InputTags;
+	RawTagString.ParseIntoArray(InputTags, TEXT(","));
+	Algo::ForEach(InputTags, [](FString& String)
+		{
+			String.TrimStartAndEndInline();
+			String.RemoveFromStart("[");
+			String.RemoveFromEnd("]");
+		});
+	return CheckAllTagsIsIn(TestTags, InputTags);
+}
+
+FString OnlineSubsystemAutoReg::GenerateTags(const FString& ServiceName, const FReportingSkippableTags& SkippableTags, const TCHAR* InTag)
+{
+	//Copy String here for ease-of-manipulation
+	FString RawInTag = InTag;
+
+	TArray<FString> TestTagsArray;
+	RawInTag.ParseIntoArray(TestTagsArray, TEXT("]"));
+	Algo::ForEach(TestTagsArray, [](FString& String)
+		{
+			String.TrimStartAndEndInline();
+			String.RemoveFromStart("[");
+		});
+	Algo::Sort(TestTagsArray);
+
+	// Search if we need to append [!mayfail] tag to indicate to 
+	// catch2 this test is in a in-development phase and failures 
+	// should be ignored.
+	for (const FString& FailableTags : SkippableTags.MayFailTags)
+	{
+		if (CheckAllTagsIsIn(TestTagsArray, FailableTags))
+		{
+			RawInTag.Append(TEXT("[!mayfail]"));
+			break;
+		}
+	}
+
+	// Search if we need to append [!shouldfail] tag to indicate to 
+	// catch2 this test should fail, and if it ever passes we should
+	// should fail.
+	for (const FString& FailableTags : SkippableTags.ShouldFailTags)
+	{
+		if (CheckAllTagsIsIn(TestTagsArray, FailableTags))
+		{
+			RawInTag.Append(TEXT("[!shouldfail]"));
+			break;
+		}
+	}
+
+	return FString::Printf(TEXT("[%s] %s"), *ServiceName, *RawInTag);
+}
+
+bool OnlineSubsystemAutoReg::ShouldDisableTest(const FString& ServiceName, const FReportingSkippableTags& SkippableTags, const TCHAR* InTag)
+{
+	//Copy String here for ease-of-manipulation
+	const FString RawInTag = InTag;
+
+	TArray<FString> TestTagsArray;
+	RawInTag.ParseIntoArray(TestTagsArray, TEXT("]"));
+	Algo::ForEach(TestTagsArray, [](FString& String)
+		{
+			String.TrimStartAndEndInline();
+			String.RemoveFromStart("[");
+		});
+	Algo::Sort(TestTagsArray);
+
+	// If we contain [!<service>] it means we shouldn't run this
+	// test against this service.
+	if (RawInTag.Contains("!" + ServiceName))
+	{
+		return true;
+	}
+
+	// If we contain tags from config it means 
+	// we shouldn't run this test
+	for (const FString& DisableTag : SkippableTags.DisableTestTags)
+	{
+		if (CheckAllTagsIsIn(TestTagsArray, DisableTag))
+		{
+			return true;
+		}
+	}
+
+	// We should run the test!
+	return false;
+}
+
+// This code is kept identical to Catch internals so that there is as little deviation from OSS_TESTS and Online_OSS_TESTS as possible
+OnlineSubsystemAutoReg::OnlineSubsystemAutoReg(OnlineSubsystemTestConstructor TestCtor, Catch::SourceLineInfo LineInfo, const char* Name, const char* Tags, const char* AddlOnlineInfo)
+{
+	auto GlobalInitalizersPtr = GetGlobalInitalizers();
+	ensure(GlobalInitalizersPtr);
+	GlobalInitalizersPtr->Add([=, this]() -> void
+		{
+			for (const FString& Subsystem : GetApplicableSubsystems())
+			{
+				FString ReportingCategory = FString::Printf(TEXT("TestReporting %s"), *Subsystem);
+				FReportingSkippableTags SkippableTags;
+				GConfig->GetArray(*ReportingCategory, TEXT("MayFailTestTags"), SkippableTags.MayFailTags, GEngineIni);
+				GConfig->GetArray(*ReportingCategory, TEXT("ShouldFailTestTags"), SkippableTags.ShouldFailTags, GEngineIni);
+				GConfig->GetArray(*ReportingCategory, TEXT("DisableTestTags"), SkippableTags.DisableTestTags, GEngineIni);
+
+				auto NewName = StringCast<ANSICHAR>(*FString::Printf(TEXT("[%s] %s"), *Subsystem, ANSI_TO_TCHAR(Name)));
+				auto NewTags = StringCast<ANSICHAR>(*GenerateTags(Subsystem, SkippableTags, ANSI_TO_TCHAR(Tags)));
+
+				// If we have tags present indicating we should not enable the test at all
+				if (ShouldDisableTest(Subsystem, SkippableTags, ANSI_TO_TCHAR(NewTags.Get())))
+				{
+					continue;
+				}
+
+				// TestCtor will create a new instance of the test we are calling- ConstructInternal is separate so that we can pass any arguments we want instead of baking them into the macro
+				OnlineSubsystemTestBase* NewTest = TestCtor();
+				NewTest->ConstructInternal(Subsystem);
+
+				// This code is lifted from Catch internals to register a test
+				Catch::getMutableRegistryHub().registerTest(Catch::makeTestCaseInfo(
+					std::string(Catch::StringRef()),  // Used for testing a static method instead of a function- not needed since we're passing an ITestInvoker macro
+					Catch::NameAndTags{ NewName.Get(), NewTags.Get() },
+					LineInfo),
+					Catch::Detail::unique_ptr(NewTest) // This is taking the ITestInvoker macro and will call invoke() to run the test
+				);
+			}
+		});
+}
\ No newline at end of file
diff --git a/Engine/Source/Programs/Online/OSSTestsCore/OnlineSubsystemCatchHelper.h b/Engine/Source/Programs/Online/OSSTestsCore/OnlineSubsystemCatchHelper.h
new file mode 100644
--- /dev/null
+++ b/Engine/Source/Programs/Online/OSSTestsCore/OnlineSubsystemCatchHelper.h
@@ -0,0 +1,102 @@
+// Copyright Epic Games, Inc. All Rights Reserved.
+
+#pragma once
+
+#include "TestHarness.h"
+#include "TestDriver.h"
+#include "Interfaces/OnlineIdentityInterface.h"
+
+#include <catch2/catch_test_case_info.hpp>
+#include <catch2/interfaces/catch_interfaces_registry_hub.hpp>
+
+class OnlineSubsystemTestBase : public Catch::ITestInvoker
+{
+public:
+	void ConstructInternal(FString ServiceName);
+	virtual ~OnlineSubsystemTestBase();
+protected:
+
+	OnlineSubsystemTestBase();
+	
+	FString GetSubsystem() const;
+	FOnlineAccountCredentials GetCredentials(int LocalUserNum) const;
+	FTestPipeline& GetLoginPipeline(uint32 NumUsersToLogin = 1) const;
+	FTestPipeline& GetPipeline();
+	void RunToCompletion() const;
+
+	/* ITestInvoker */
+	virtual void invoke() const override = 0;
+
+private:
+	FString Tags;
+	FString Subsystem;
+	// Catch's ITestInvoker is a const interface but we'll be changing stuff (emplacing steps into the driver, setting flags, etc.) so we make these mutable
+	mutable FTestDriver Driver;
+	mutable FTestPipeline Pipeline;
+	mutable uint32 NumLocalUsers = -1;
+};
+
+typedef OnlineSubsystemTestBase* (*OnlineSubsystemTestConstructor)();
+
+class OnlineSubsystemAutoReg
+{
+private:
+	// TODO: This should poll some info from the tags and generate the list based on that
+	TArray<FString> GetApplicableSubsystems();
+
+public:
+	struct FReportingSkippableTags
+	{
+		TArray<FString> MayFailTags;
+		TArray<FString> ShouldFailTags;
+		TArray<FString> DisableTestTags;
+	};
+
+	/*
+	* Helper function that calls CheckAllTagsIsIn(const TArray<FString>&, const TArray<FString>&);
+	*
+	* @param  TestTags			The array of test tags we wish to test against.
+	* @param  RawTagString		Comma seperated string of elemnets we wish to convert to an array.
+	*
+	* @return true if all elements of RawTagString prased as an comma sperated array is in TestTags.
+	*/
+	static bool CheckAllTagsIsIn(const TArray<FString>& TestTags, const FString& RawTagString);
+
+
+	/**
+	 * Checks if every element of InputTags is in TestTags.
+	 *
+	 * @param  TestTags       The array of test tags we wish to test against.
+	 * @param  InputTags	  The array of tags we wish to see if all elements of are present in TestTags.
+	 *
+	 * @return  true if all elements of InputTags is in TestTags.
+	 */
+	static bool CheckAllTagsIsIn(const TArray<FString>& TestTags, const TArray<FString>& InputTags);
+
+	static FString GenerateTags(const FString& ServiceName, const FReportingSkippableTags& SkippableTags, const TCHAR* InTag);
+	static bool ShouldDisableTest(const FString& ServiceName, const FReportingSkippableTags& SkippableTags, const TCHAR* InTag);
+
+	// This code is kept identical to Catch internals so that there is as little deviation from OSS_TESTS and Online_OSS_TESTS as possible
+	OnlineSubsystemAutoReg(OnlineSubsystemTestConstructor TestCtor, Catch::SourceLineInfo LineInfo, const char* Name, const char* Tags, const char* AddlOnlineInfo);
+};
+
+TArray<TFunction<void()>>* GetGlobalInitalizers();
+
+
+// INTERNAL_CATCH_UNIQUE_NAME is what Catch uses to translate the name of the test into a unique identifier without any spaces
+
+#define INTERNAL_ONLINESUBSYSTEM_TEST_CASE_NAMED(RegName, Name, Tags, ...)\
+namespace {\
+class PREPROCESSOR_JOIN(OnlineSubsystemTest_,RegName) : public OnlineSubsystemTestBase {\
+protected:\
+	virtual void invoke() const override;\
+};\
+	OnlineSubsystemTestBase* PREPROCESSOR_JOIN(Construct_,RegName) (){\
+		return new PREPROCESSOR_JOIN(OnlineSubsystemTest_,RegName) ();\
+	}\
+	OnlineSubsystemAutoReg RegName( PREPROCESSOR_JOIN(&Construct_,RegName) , CATCH_INTERNAL_LINEINFO , Name, Tags, "");\
+}\
+void PREPROCESSOR_JOIN(OnlineSubsystemTest_,RegName)::invoke() const\
+
+#define ONLINESUBSYSTEM_TEST_CASE(Name, Tags, ...) \
+	INTERNAL_ONLINESUBSYSTEM_TEST_CASE_NAMED(INTERNAL_CATCH_UNIQUE_NAME(OnlineSubsystemRegistrar), Name, Tags, __VA_ARGS__)
diff --git a/Engine/Source/Programs/Online/OSSTestsCore/Tests/OSSCatchHelpers/OSSCatchHelpers.cpp b/Engine/Source/Programs/Online/OSSTestsCore/Tests/OSSCatchHelpers/OSSCatchHelpers.cpp
new file mode 100644
--- /dev/null
+++ b/Engine/Source/Programs/Online/OSSTestsCore/Tests/OSSCatchHelpers/OSSCatchHelpers.cpp
@@ -0,0 +1,229 @@
+// Copyright Epic Games, Inc. All Rights Reserved.
+
+#include "OnlineSubsystemCatchHelper.h"
+
+#define SELFTEST_TAG "[selftest][osscatchhelper]"
+#define GENERATETAGS_TAG "[generatetags]"
+#define CHECKALLTAGSISIN_TAG "[checkalltagsisin]"
+#define SHOULDDISABLETEST_TAG "[shoulddisabletest]"
+#define SELFTEST_TEST_CASE(x, ...) TEST_CASE(x, SELFTEST_TAG __VA_ARGS__)
+
+SELFTEST_TEST_CASE("GenerateTags append MayFailTags case", GENERATETAGS_TAG)
+{
+	FString TestTags = "[foo][bar]";
+	OnlineSubsystemAutoReg::FReportingSkippableTags SkipTags;
+	SkipTags.MayFailTags = { "[foo]" };
+	FString OutTags = OnlineSubsystemAutoReg::GenerateTags("TestService", SkipTags, *TestTags);
+	CHECK(OutTags == TEXT("[TestService] [foo][bar][!mayfail]"));
+}
+
+SELFTEST_TEST_CASE("GenerateTags append multiple match MayFailTags case", GENERATETAGS_TAG)
+{
+	FString TestTags = "[foo][bar]";
+	OnlineSubsystemAutoReg::FReportingSkippableTags SkipTags;
+	SkipTags.MayFailTags = { "[bar]", "[foo]" };
+	FString OutTags = OnlineSubsystemAutoReg::GenerateTags("TestService", SkipTags, *TestTags);
+	CHECK(OutTags == TEXT("[TestService] [foo][bar][!mayfail]"));
+}
+
+SELFTEST_TEST_CASE("GenerateTags append by last match MayFailTags case", GENERATETAGS_TAG)
+{
+	FString TestTags = "[foo][bar]";
+	OnlineSubsystemAutoReg::FReportingSkippableTags SkipTags;
+	SkipTags.MayFailTags = { "[wiz]", "[foo]" };
+	FString OutTags = OnlineSubsystemAutoReg::GenerateTags("TestService", SkipTags, *TestTags);
+	CHECK(OutTags == TEXT("[TestService] [foo][bar][!mayfail]"));
+}
+
+SELFTEST_TEST_CASE("GenerateTags append by last match mutli-tag case", GENERATETAGS_TAG)
+{
+	FString TestTags = "[foo][bar]";
+	OnlineSubsystemAutoReg::FReportingSkippableTags SkipTags;
+	SkipTags.MayFailTags = { "[foo],bar" };
+	FString OutTags = OnlineSubsystemAutoReg::GenerateTags("TestService", SkipTags, *TestTags);
+	CHECK(OutTags == TEXT("[TestService] [foo][bar][!mayfail]"));
+}
+
+SELFTEST_TEST_CASE("GenerateTags append by last match mutli-tag no match case", GENERATETAGS_TAG)
+{
+	FString TestTags = "[foo][bar]";
+	OnlineSubsystemAutoReg::FReportingSkippableTags SkipTags;
+	SkipTags.MayFailTags = { "[foo],[wiz]" };
+	FString OutTags = OnlineSubsystemAutoReg::GenerateTags("TestService", SkipTags, *TestTags);
+	CHECK(OutTags == TEXT("[TestService] [foo][bar]"));
+}
+
+SELFTEST_TEST_CASE("GenerateTags don't append MayFailTags no match case", GENERATETAGS_TAG)
+{
+	OnlineSubsystemAutoReg::FReportingSkippableTags SkipTags;
+	SkipTags.MayFailTags = { "[wiz]" };
+	FString TestTags = "[foo][bar]";
+	FString OutTags = OnlineSubsystemAutoReg::GenerateTags("TestService", SkipTags, *TestTags);
+	CHECK(OutTags == TEXT("[TestService] [foo][bar]"));
+}
+
+SELFTEST_TEST_CASE("GenerateTags append ShouldFail case", GENERATETAGS_TAG)
+{
+	FString TestTags = "[foo][bar]";
+	OnlineSubsystemAutoReg::FReportingSkippableTags SkipTags;
+	SkipTags.ShouldFailTags = { "[foo]" };
+	FString OutTags = OnlineSubsystemAutoReg::GenerateTags("TestService", SkipTags, *TestTags);
+	CHECK(OutTags == TEXT("[TestService] [foo][bar][!shouldfail]"));
+}
+
+SELFTEST_TEST_CASE("GenerateTags append multiple match ShouldFail case", GENERATETAGS_TAG)
+{
+	FString TestTags = "[foo][bar]";
+	OnlineSubsystemAutoReg::FReportingSkippableTags SkipTags;
+	SkipTags.ShouldFailTags = { "[bar]", "[foo]" };
+	FString OutTags = OnlineSubsystemAutoReg::GenerateTags("TestService", SkipTags, *TestTags);
+	CHECK(OutTags == TEXT("[TestService] [foo][bar][!shouldfail]"));
+}
+
+SELFTEST_TEST_CASE("GenerateTags append by last match ShouldFail case", GENERATETAGS_TAG)
+{
+	FString TestTags = "[foo][bar]";
+	OnlineSubsystemAutoReg::FReportingSkippableTags SkipTags;
+	SkipTags.ShouldFailTags = { "[wiz]", "[foo]" };
+	FString OutTags = OnlineSubsystemAutoReg::GenerateTags("TestService", SkipTags, *TestTags);
+	CHECK(OutTags == TEXT("[TestService] [foo][bar][!shouldfail]"));
+}
+
+SELFTEST_TEST_CASE("GenerateTags don't append ShouldFail no match case", GENERATETAGS_TAG)
+{
+	FString TestTags = "[foo][bar]";
+	OnlineSubsystemAutoReg::FReportingSkippableTags SkipTags;
+	SkipTags.ShouldFailTags = { "[wiz]" };
+	FString OutTags = OnlineSubsystemAutoReg::GenerateTags("TestService", SkipTags, *TestTags);
+	CHECK(OutTags == TEXT("[TestService] [foo][bar]"));
+}
+
+SELFTEST_TEST_CASE("GenerateTags appends MayFail and ShouldFail Case", GENERATETAGS_TAG)
+{
+	OnlineSubsystemAutoReg::FReportingSkippableTags SkipTags;
+	SkipTags.MayFailTags = { "[foo]" };
+	SkipTags.ShouldFailTags = { "[bar]" };
+	FString TestTags = "[foo][bar]";
+	FString OutTags = OnlineSubsystemAutoReg::GenerateTags("TestService", SkipTags, *TestTags);
+	CHECK(OutTags == TEXT("[TestService] [foo][bar][!mayfail][!shouldfail]"));
+}
+
+
+SELFTEST_TEST_CASE("GenerateTags append by last match mutli-tag ShouldFail case", GENERATETAGS_TAG)
+{
+	FString TestTags = "[foo][bar]";
+	OnlineSubsystemAutoReg::FReportingSkippableTags SkipTags;
+	SkipTags.ShouldFailTags = { "[foo],bar" };
+	FString OutTags = OnlineSubsystemAutoReg::GenerateTags("TestService", SkipTags, *TestTags);
+	CHECK(OutTags == TEXT("[TestService] [foo][bar][!shouldfail]"));
+}
+
+SELFTEST_TEST_CASE("GenerateTags append by last match mutli-tag no match ShouldFail case", GENERATETAGS_TAG)
+{
+	FString TestTags = "[foo][bar]";
+	OnlineSubsystemAutoReg::FReportingSkippableTags SkipTags;
+	SkipTags.ShouldFailTags = { "[foo],[wiz]" };
+	FString OutTags = OnlineSubsystemAutoReg::GenerateTags("TestService", SkipTags, *TestTags);
+	CHECK(OutTags == TEXT("[TestService] [foo][bar]"));
+}
+
+SELFTEST_TEST_CASE("GenerateTags appends no tags case", GENERATETAGS_TAG)
+{
+	OnlineSubsystemAutoReg::FReportingSkippableTags SkipTags;
+	SkipTags.MayFailTags = { "[wiz]" };
+	SkipTags.ShouldFailTags = {  };
+	FString TestTags = "[foo][bar]";
+	FString OutTags = OnlineSubsystemAutoReg::GenerateTags("TestService", SkipTags, *TestTags);
+	CHECK(OutTags == TEXT("[TestService] [foo][bar]"));
+}
+
+SELFTEST_TEST_CASE("ShouldDisableTest returns true on single-tag config", SHOULDDISABLETEST_TAG)
+{
+	OnlineSubsystemAutoReg::FReportingSkippableTags SkipTags;
+	SkipTags.DisableTestTags = { "[foo]" };
+	FString TestTags = "[foo][bar]";
+	CHECK(OnlineSubsystemAutoReg::ShouldDisableTest("TestService", SkipTags, *TestTags) == true);
+}
+
+SELFTEST_TEST_CASE("ShouldDisableTest returns true on multi-tag config", SHOULDDISABLETEST_TAG)
+{
+	OnlineSubsystemAutoReg::FReportingSkippableTags SkipTags;
+	SkipTags.DisableTestTags = { "[foo],bar" };
+	FString TestTags = "[foo][bar]";
+	CHECK(OnlineSubsystemAutoReg::ShouldDisableTest("TestService", SkipTags, *TestTags) == true);
+}
+
+SELFTEST_TEST_CASE("ShouldDisableTest returns true on !<service>", SHOULDDISABLETEST_TAG)
+{
+	FString TestTags = "[foo][bar][!TestService]";
+	CHECK(OnlineSubsystemAutoReg::ShouldDisableTest("TestService", {}, *TestTags) == true);
+}
+
+SELFTEST_TEST_CASE("ShouldDisableTest returns false with no tags and no config skips", SHOULDDISABLETEST_TAG)
+{
+	FString TestTags = "[foo][bar]";
+	CHECK(OnlineSubsystemAutoReg::ShouldDisableTest("TestService", {}, *TestTags) == false);
+}
+
+SELFTEST_TEST_CASE("ShouldDisableTest returns false with no matching no-tags and no matching multi-tag config skips", SHOULDDISABLETEST_TAG)
+{
+	OnlineSubsystemAutoReg::FReportingSkippableTags SkipTags;
+	SkipTags.DisableTestTags = { "[foo],wiz" };
+	FString TestTags = "[foo][bar]";
+	CHECK(OnlineSubsystemAutoReg::ShouldDisableTest("TestService", SkipTags, *TestTags) == false);
+}
+
+SELFTEST_TEST_CASE("CheckAllTagsIsIn(TArray, FString) true cases", CHECKALLTAGSISIN_TAG)
+{
+	TArray<FString> TestTags = { "bob", "alice", "foo" };
+	CAPTURE(TestTags);
+
+	// Truthy Cases
+	CHECK(OnlineSubsystemAutoReg::CheckAllTagsIsIn(TestTags, "bob, alice") == true);
+	CHECK(OnlineSubsystemAutoReg::CheckAllTagsIsIn(TestTags, "bob,alice") == true);
+	CHECK(OnlineSubsystemAutoReg::CheckAllTagsIsIn(TestTags, " bob,alice ") == true);
+	CHECK(OnlineSubsystemAutoReg::CheckAllTagsIsIn(TestTags, "foo") == true);
+	CHECK(OnlineSubsystemAutoReg::CheckAllTagsIsIn(TestTags, ",foo") == true);
+	CHECK(OnlineSubsystemAutoReg::CheckAllTagsIsIn(TestTags, "bob,alice,foo") == true);
+
+	//Bracket Parsing
+	CHECK(OnlineSubsystemAutoReg::CheckAllTagsIsIn(TestTags, "[bob],[alice],[foo]") == true);
+	CHECK(OnlineSubsystemAutoReg::CheckAllTagsIsIn(TestTags, "[bob], [alice,foo]") == true);
+	CHECK(OnlineSubsystemAutoReg::CheckAllTagsIsIn(TestTags, "bob],  alice],  [foo]  ,") == true);
+	CHECK(OnlineSubsystemAutoReg::CheckAllTagsIsIn(TestTags, ",[foo]") == true);
+	CHECK(OnlineSubsystemAutoReg::CheckAllTagsIsIn(TestTags, ",foo]") == true);
+	CHECK(OnlineSubsystemAutoReg::CheckAllTagsIsIn(TestTags, "[wiz]") == false);
+
+	// Negative Cases
+	CHECK(OnlineSubsystemAutoReg::CheckAllTagsIsIn(TestTags, "bob,alice,foo,wiz") == false);
+	CHECK(OnlineSubsystemAutoReg::CheckAllTagsIsIn(TestTags, "bob,wiz") == false);
+	CHECK(OnlineSubsystemAutoReg::CheckAllTagsIsIn(TestTags, ",wiz") == false);
+
+	// Bound Checks
+	CHECK(OnlineSubsystemAutoReg::CheckAllTagsIsIn(TestTags, ",") == false);
+	CHECK(OnlineSubsystemAutoReg::CheckAllTagsIsIn(TestTags, "") == false);
+	CHECK(OnlineSubsystemAutoReg::CheckAllTagsIsIn({}, "") == false);
+}
+
+
+SELFTEST_TEST_CASE("CheckAllTagsIsIn(TArray, TArray) true cases", CHECKALLTAGSISIN_TAG)
+{
+	TArray<FString> TestTags = { "bob", "alice", "foo" };
+	CAPTURE(TestTags);
+
+	// Truthy Cases
+	CHECK(OnlineSubsystemAutoReg::CheckAllTagsIsIn(TestTags, { "bob", "alice" }) == true);
+	CHECK(OnlineSubsystemAutoReg::CheckAllTagsIsIn(TestTags, { "bob", "alice", "foo" }) == true);
+	CHECK(OnlineSubsystemAutoReg::CheckAllTagsIsIn(TestTags, TArray<FString>({ "foo" })) == true);
+
+	// Negative Cases
+	CHECK(OnlineSubsystemAutoReg::CheckAllTagsIsIn(TestTags, TArray<FString>({ "wiz" })) == false);
+	CHECK(OnlineSubsystemAutoReg::CheckAllTagsIsIn(TestTags, { "bob", "alice", "foo", "wiz" }) == false);
+	CHECK(OnlineSubsystemAutoReg::CheckAllTagsIsIn(TestTags, { "bob", "alice", "wiz" }) == false);
+
+	// Bounds
+	CHECK(OnlineSubsystemAutoReg::CheckAllTagsIsIn(TestTags, TArray<FString>({})) == false);
+	CHECK(OnlineSubsystemAutoReg::CheckAllTagsIsIn({}, TArray<FString>({})) == false);
+	CHECK(OnlineSubsystemAutoReg::CheckAllTagsIsIn({}, TArray<FString>({ "wiz" })) == false);
+	CHECK(OnlineSubsystemAutoReg::CheckAllTagsIsIn({ "wiz" }, TArray<FString>({})) == false);
+}
\ No newline at end of file
diff --git a/Engine/Source/Programs/Online/OSSTestsCore/Tests/Purchase/PurchaseGetReceiptsTests.cpp b/Engine/Source/Programs/Online/OSSTestsCore/Tests/Purchase/PurchaseGetReceiptsTests.cpp
new file mode 100644
--- /dev/null
+++ b/Engine/Source/Programs/Online/OSSTestsCore/Tests/Purchase/PurchaseGetReceiptsTests.cpp
@@ -0,0 +1,37 @@
+// Copyright Epic Games, Inc. All Rights Reserved.
+
+#include "TestDriver.h"
+#include "TestUtilities.h"
+
+#include "OnlineSubsystemCatchHelper.h"
+
+#include "Helpers/Identity/IdentityGetUniquePlayerIdHelper.h"
+#include "Helpers/Purchase/PurchaseGetReceiptsHelper.h"
+#include "Helpers/Purchase/PurchaseQueryReceiptsHelper.h"
+#include "Helpers/Purchase/PurchaseCheckoutHelper.h"
+
+#define PURCHASE_TAG "[suite_purchase]"
+#define EG_PURCHASE_GETRECEIPTS_TAG PURCHASE_TAG "[getreceipts]"
+
+#define PURCHASE_TEST_CASE(x, ...) ONLINESUBSYSTEM_TEST_CASE(x, PURCHASE_TAG __VA_ARGS__)
+
+PURCHASE_TEST_CASE("Verify calling GetReceipts with valid inputs returns the expected result(Success Case)", EG_PURCHASE_GETRECEIPTS_TAG)
+{
+	int32 LocalUserNum = 0;
+	FUniqueNetIdPtr LocalUserId = nullptr;
+	int32 LocalQuantity = 1;
+	bool bLocalIsConsumable = true;
+	const FOfferNamespace LocalOfferNamespace = TEXT("");
+	const FUniqueOfferId& LocalOfferId = TEXT("Item1_Id");
+	const FString LocalItemName = TEXT("Cool Item1");
+	FPurchaseCheckoutRequest PurchCheckRequest = {};
+	PurchCheckRequest.AddPurchaseOffer(LocalOfferNamespace, LocalOfferId, LocalQuantity, bLocalIsConsumable);
+	int32 NumUsersToImplicitLogin = 1;
+
+	GetLoginPipeline(NumUsersToImplicitLogin)
+		.EmplaceStep<FIdentityGetUniquePlayerIdStep>(LocalUserNum, [&LocalUserId](FUniqueNetIdPtr InUserId) {LocalUserId = InUserId; })
+		.EmplaceStep<FPurchaseCheckoutStep>(&LocalUserId, &PurchCheckRequest, LocalOfferId, LocalItemName)
+		.EmplaceStep<FPurchaseGetReceiptsStep>(&LocalUserId, LocalOfferId, LocalItemName);
+	
+	RunToCompletion();
+}
diff --git a/Engine/Source/Programs/Online/OSSTestsCore/Tests/Purchase/PurchaseQueryReceiptsTests.cpp b/Engine/Source/Programs/Online/OSSTestsCore/Tests/Purchase/PurchaseQueryReceiptsTests.cpp
new file mode 100644
--- /dev/null
+++ b/Engine/Source/Programs/Online/OSSTestsCore/Tests/Purchase/PurchaseQueryReceiptsTests.cpp
@@ -0,0 +1,39 @@
+// Copyright Epic Games, Inc. All Rights Reserved.
+
+#include "TestDriver.h"
+#include "TestUtilities.h"
+
+#include "OnlineSubsystemCatchHelper.h"
+
+#include "Helpers/Purchase/PurchaseQueryReceiptsHelper.h"
+#include "Helpers/Purchase/PurchaseGetReceiptsHelper.h"
+#include "Helpers/Purchase/PurchaseCheckoutHelper.h"
+#include "Helpers/Identity/IdentityGetUniquePlayerIdHelper.h"
+
+#define PURCHASE_TAG "[suite_purchase]"
+#define EG_PURCHASE_QUERYRECEIPTS_TAG PURCHASE_TAG "[queryreceipts]"
+
+#define PURCHASE_TEST_CASE(x, ...) ONLINESUBSYSTEM_TEST_CASE(x, PURCHASE_TAG __VA_ARGS__)
+
+PURCHASE_TEST_CASE("Verify calling QueryReceipts with valid inputs returns the expected result(Success Case)", EG_PURCHASE_QUERYRECEIPTS_TAG)
+{
+	int32 LocalUserNum = 0;
+	FUniqueNetIdPtr LocalUserId = nullptr;
+	bool LocalRestoreReceipts = true; //this param doesen't use in related step
+	const FUniqueOfferId& LocalOfferId = TEXT("Item1_Id");
+	const FString LocalItemName = TEXT("Cool Item1");
+	int32 LocalQuantity = 1;
+	bool bLocalIsConsumable = true;
+	const FOfferNamespace LocalOfferNamespace = TEXT("");
+	FPurchaseCheckoutRequest PurchCheckRequest = {};
+	PurchCheckRequest.AddPurchaseOffer(LocalOfferNamespace, LocalOfferId, LocalQuantity, bLocalIsConsumable);
+	int32 NumUsersToImplicitLogin = 1;
+
+	GetLoginPipeline(NumUsersToImplicitLogin)
+		.EmplaceStep<FIdentityGetUniquePlayerIdStep>(LocalUserNum, [&LocalUserId](FUniqueNetIdPtr InUserId) {LocalUserId = InUserId; })
+		.EmplaceStep<FPurchaseCheckoutStep>(&LocalUserId, &PurchCheckRequest, LocalOfferId, LocalItemName)
+		.EmplaceStep<FPurchaseQueryReceiptsStep>(&LocalUserId, LocalRestoreReceipts, LocalOfferId)
+		.EmplaceStep<FPurchaseGetReceiptsStep>(&LocalUserId, LocalOfferId, LocalItemName);
+
+	RunToCompletion();
+}
\ No newline at end of file
