Apply fixes UE-207122,  UE-207036 to Release-5.4 branch.
- Fixes preprocessor directive issue at OnlineCatchHelper.cpp file.
- Fixes MacroRedefinition in OnlineTestsCore.Build.cs

diff --git a/Engine/Source/Programs/Online/OnlineTestsCore/Source/OnlineTestsCore.Build.cs b/Engine/Source/Programs/Online/OnlineTestsCore/Source/OnlineTestsCore.Build.cs
--- a/Engine/Source/Programs/Online/OnlineTestsCore/Source/OnlineTestsCore.Build.cs
+++ b/Engine/Source/Programs/Online/OnlineTestsCore/Source/OnlineTestsCore.Build.cs
@@ -8,11 +8,6 @@ using System.Linq;
 
 public class OnlineTestsCore : ModuleRules
 {
-	public virtual bool bRequireApplicationTick { get { return false; } }
-	public virtual bool bRequirePlatformInit { get { return false; } }
-	public virtual string PlatformFileName { get { return ""; } }
-	public virtual bool bUseExternAuth { get { return false; } }
-
 	public OnlineTestsCore(ReadOnlyTargetRules Target) : base(Target)
 	{
 		bTreatAsEngineModule = false;
@@ -40,10 +35,6 @@ public class OnlineTestsCore : ModuleRules
 			}
 		);
 
-		PublicDefinitions.Add(String.Format("ONLINETESTS_REQUIREAPPLICATIONTICK={0}", bRequireApplicationTick ? 1 : 0));
-		PublicDefinitions.Add(String.Format("ONLINETESTS_REQUIREPLATFORMINIT={0}", bRequirePlatformInit ? 1 : 0));
-		PublicDefinitions.Add(String.Format("ONLINETESTS_PLATFORMFILENAME={0}", PlatformFileName));
-
 		// Disable external auth if target doesn't define it.
 		if (!Target.GlobalDefinitions.Contains("ONLINETESTS_USEEXTERNAUTH=1"))
 		{
diff --git a/Engine/Source/Programs/Online/OnlineTestsCore/Source/Private/OnlineCatchHelper.cpp b/Engine/Source/Programs/Online/OnlineTestsCore/Source/Private/OnlineCatchHelper.cpp
--- a/Engine/Source/Programs/Online/OnlineTestsCore/Source/Private/OnlineCatchHelper.cpp
+++ b/Engine/Source/Programs/Online/OnlineTestsCore/Source/Private/OnlineCatchHelper.cpp
@@ -81,9 +81,9 @@ bool OnlineTestBase::DeleteAccountsForCurrentTemplate() const
 {
 #if ONLINETESTS_USEEXTERNAUTH
 	return CustomDeleteAccounts();
-#else //ONLINETESTS_USEEXTERNAUTH
+#else // ONLINETESTS_USEEXTERNAUTH
 	return false;
-#endif ONLINETESTS_USEEXTERNAUTH
+#endif // ONLINETESTS_USEEXTERNAUTH
 }
 
 void OnlineTestBase::DestroyCurrentServiceModule() const
@@ -95,9 +95,9 @@ bool OnlineTestBase::ResetAccountStatus() const
 {
 #if ONLINETESTS_USEEXTERNAUTH
 	return CustomResetAccounts();
-#else //ONLINETESTS_USEEXTERNAUTH
+#else // ONLINETESTS_USEEXTERNAUTH
 	return false;
-#endif ONLINETESTS_USEEXTERNAUTH
+#endif // ONLINETESTS_USEEXTERNAUTH
 }
 
 TArray<FString> GetServiceModules()
@@ -172,9 +172,9 @@ FAuthLogin::Params OnlineTestBase::GetCredentials(int LocalUserNum) const
 {
 #if ONLINETESTS_USEEXTERNAUTH
 	return CustomCredentials(LocalUserNum);
-#else //ONLINETESTS_USEEXTERNAUTH
+#else // ONLINETESTS_USEEXTERNAUTH
 	return GetIniCredentials(LocalUserNum);
-#endif ONLINETESTS_USEEXTERNAUTH
+#endif // ONLINETESTS_USEEXTERNAUTH
 }
 
 FString OnlineTestBase::GetLoginCredentialCategory() const
