Fix warning: MCP: Open slots count mismatch.

- Corrects the number of open connections for test pipelining.

diff --git a/Engine/Source/Programs/Online/OSSTestsCore/Tests/Session/SessionUpdateSessionTests.cpp b/Engine/Source/Programs/Online/OSSTestsCore/Tests/Session/SessionUpdateSessionTests.cpp
--- a/Engine/Source/Programs/Online/OSSTestsCore/Tests/Session/SessionUpdateSessionTests.cpp
+++ b/Engine/Source/Programs/Online/OSSTestsCore/Tests/Session/SessionUpdateSessionTests.cpp
@@ -22,7 +22,7 @@ SESSION_TEST_CASE("Verify calling Session UpdateSession with valid inputs return
 {
 	int32 LocalUserNum = 0;
 	int32 PublicConnections = 1;
-	int32 NewPublicConnections = 2;
+	int32 NewPublicConnections = 1;
 	int32 NumUsersToImplicitLogin = 1;
 	FUniqueNetIdPtr UserId = nullptr;
 	FName SessionName = FName(FString::Printf(TEXT("TestSession_%s"), *FTestUtilities::GetUniqueTestString()));
