Fix warning Mcp: GetFriendsList and RejectInvite request failed

- added lowering of the logging level for the corresponding helper, since these warnings occur in case of successful execution of the operation.

diff --git a/Engine/Source/Programs/Online/OSSTestsCore/Helpers/Friends/FriendsEnsureFriendshipHelper.h b/Engine/Source/Programs/Online/OSSTestsCore/Helpers/Friends/FriendsEnsureFriendshipHelper.h
--- a/Engine/Source/Programs/Online/OSSTestsCore/Helpers/Friends/FriendsEnsureFriendshipHelper.h
+++ b/Engine/Source/Programs/Online/OSSTestsCore/Helpers/Friends/FriendsEnsureFriendshipHelper.h
@@ -32,9 +32,21 @@ struct FFriendsEnsureFriendshipStep : public FTestPipeline::FStep
 			.EmplaceInnerStep<FFriendsReadFriendsListStep>(InTargetUserNum, InLocalListName, InbIsFriendsListPopulated)
 			.EmplaceInnerStep<FFriendsIsFriendStep>(InTargetUserNum, InLocalUserId, InLocalListName)
 			.Steps)
+		, OldVerbosity(LogOnlineFriend.GetVerbosity())
+		, CurrentVerbosity([]()
+			{
+				LogOnlineFriend.SetVerbosity(ELogVerbosity::Error);
+				return ELogVerbosity::Error;
+			}())
 	{}
 
-	virtual ~FFriendsEnsureFriendshipStep() = default;
+	virtual ~FFriendsEnsureFriendshipStep()
+	{
+		if (OldVerbosity != LogOnlineFriend.GetVerbosity())
+		{
+			LogOnlineFriend.SetVerbosity(OldVerbosity);
+		}
+	}
 
 	virtual EContinuance Tick(IOnlineSubsystem* OnlineSubsystem) override
 	{
@@ -45,7 +57,7 @@ struct FFriendsEnsureFriendshipStep : public FTestPipeline::FStep
 
 		OnlineFriendsPtr = OnlineSubsystem->GetFriendsInterface();
 		REQUIRE(OnlineFriendsPtr != nullptr);
-		
+
 		CurrentStepContinuanceResult = InnerSteps.HeapTop()->Tick(OnlineSubsystem);
 
 		if (CurrentStepContinuanceResult == FStep::EContinuance::Done)
@@ -61,4 +73,6 @@ protected:
 	FStep::EContinuance CurrentStepContinuanceResult = FStep::EContinuance::ContinueStepping;
 	IOnlineFriendsPtr OnlineFriendsPtr = nullptr;
 	TArray<TUniquePtr<FTestPipeline::FStep>> InnerSteps;
-};
+	ELogVerbosity::Type OldVerbosity;
+	ELogVerbosity::Type CurrentVerbosity;
+};
\ No newline at end of file
