Social Interface integration tests.

- Some of the interfaces are currently disabled for two reasons: some of the functionality is disabled in EOS, some of it has no implementation in any of the OSS.

diff --git a/Engine/Plugins/Online/OnlineServicesOSSAdapter/Source/Private/Online/SocialOSSAdapter.cpp b/Engine/Plugins/Online/OnlineServicesOSSAdapter/Source/Private/Online/SocialOSSAdapter.cpp
--- a/Engine/Plugins/Online/OnlineServicesOSSAdapter/Source/Private/Online/SocialOSSAdapter.cpp
+++ b/Engine/Plugins/Online/OnlineServicesOSSAdapter/Source/Private/Online/SocialOSSAdapter.cpp
@@ -95,7 +95,7 @@ TOnlineAsyncOpHandle<FQueryFriends> FSocialOSSAdapter::QueryFriends(FQueryFriend
 		}
 		else
 		{
-			Op.SetError(Errors::InvalidUser());
+			Op.SetError(Errors::InvalidParams());
 		}
 	})
 	.Enqueue(GetSerialQueue());
@@ -151,7 +151,7 @@ TOnlineAsyncOpHandle<FSendFriendInvite> FSocialOSSAdapter::SendFriendInvite(FSen
 
 		if (LocalUserNum == INDEX_NONE || !FriendId)
 		{
-			Op.SetError(Errors::InvalidUser());
+			Op.SetError(Errors::InvalidParams());
 			return;
 		}
 
@@ -192,7 +192,7 @@ TOnlineAsyncOpHandle<FAcceptFriendInvite> FSocialOSSAdapter::AcceptFriendInvite(
 
 		if (LocalUserNum == INDEX_NONE || !FriendId)
 		{
-			Op.SetError(Errors::InvalidUser());
+			Op.SetError(Errors::InvalidParams());
 			return;
 		}
 
@@ -233,7 +233,7 @@ TOnlineAsyncOpHandle<FRejectFriendInvite> FSocialOSSAdapter::RejectFriendInvite(
 
 		if (LocalUserNum == INDEX_NONE || !FriendId)
 		{
-			Op.SetError(Errors::InvalidUser());
+			Op.SetError(Errors::InvalidParams());
 			return;
 		}
 
@@ -371,7 +371,7 @@ TOnlineAsyncOpHandle<FBlockUser> FSocialOSSAdapter::BlockUser(FBlockUser::Params
 		}
 		else
 		{
-			Op.SetError(Errors::InvalidUser());
+			Op.SetError(Errors::InvalidParams());
 		}
 	})
 	.Enqueue(GetSerialQueue());
diff --git a/Engine/Source/Programs/Online/OnlineTestsCore/Source/Private/OnlineCatchHelper.cpp b/Engine/Source/Programs/Online/OnlineTestsCore/Source/Private/OnlineCatchHelper.cpp
--- a/Engine/Source/Programs/Online/OnlineTestsCore/Source/Private/OnlineCatchHelper.cpp
+++ b/Engine/Source/Programs/Online/OnlineTestsCore/Source/Private/OnlineCatchHelper.cpp
@@ -270,12 +270,26 @@ FTestPipeline& OnlineTestBase::GetLoginPipeline(uint32 NumUsersToLogin) const
 	return *Pipeline;
 }
 
+FTestPipeline& OnlineTestBase::GetLoginPipeline(uint32 UserNumToLogin, FAccountId& OutAccountId) const
+{
+	REQUIRE(NumLocalUsers == -1); // Don't call GetLoginPipeline more than once per test
+	// Make sure input delegates are fired for adding the required user count.
+	//EnsureLocalUserCount(UserNumToLogin);
+	Pipeline->EmplaceStep<FAuthLoginStep>(GetCredentials(UserNumToLogin));
+	// Perform login so we can bulk assign users in the next step.
+	RunToCompletion(false);
+	AssignLoginUsers(UserNumToLogin, OutAccountId);
+	// Return a fresh pipeline so the logins added by GetLoginPipeline don't execute again.
+	Pipeline = MakeShared<FTestPipeline>(Driver.MakePipeline());
+	return *Pipeline;
+}
+
 FTestPipeline& OnlineTestBase::GetPipeline() const
 {
 	return GetLoginPipeline(0);
 }
 
-void OnlineTestBase::RunToCompletion(bool bLogout) const
+void OnlineTestBase::RunToCompletion(bool bLogout, const TOptional<int32> UserNumToLogout) const
 {
 	bool bUseAutoLogin = false;
 	bool bUseImplicitLogin = false;
@@ -294,6 +308,10 @@ void OnlineTestBase::RunToCompletion(bool bLogout) const
 			// NumLocalUsers = 1;
 			// Pipeline.EmplaceStep<FAuthAutoLoginStep>(0);
 		}
+		else if (UserNumToLogout.IsSet())
+		{
+			Pipeline->EmplaceStep<FAuthLogoutStep>(FPlatformMisc::GetPlatformUserForUserIndex(UserNumToLogout.GetValue()));
+		}
 		else 
 		{
 			for (uint32 i = 0; i < NumLocalUsers; i++)
diff --git a/Engine/Source/Programs/Online/OnlineTestsCore/Source/Private/Tests/Social/AcceptFriendInviteTests.cpp b/Engine/Source/Programs/Online/OnlineTestsCore/Source/Private/Tests/Social/AcceptFriendInviteTests.cpp
new file mode 100644
--- /dev/null
+++ b/Engine/Source/Programs/Online/OnlineTestsCore/Source/Private/Tests/Social/AcceptFriendInviteTests.cpp
@@ -0,0 +1,96 @@
+// Copyright Epic Games, Inc. All Rights Reserved.
+
+#include "Helpers/Social/AcceptFriendInviteHelper.h"
+#include "Helpers/Auth/AuthLogout.h"
+
+#include "OnlineCatchHelper.h"
+
+#define SOCIAL_TAG "[suite_social]"
+#define EG_SOCIAL_ACCEPTFRIENDINVITE_TAG SOCIAL_TAG "[acceptfriendinvite]"
+#define EG_SOCIAL_ACCEPTFRIENDINVITEEOS_TAG SOCIAL_TAG "[acceptfriendinvite][.EOS]"
+#define EG_SOCIAL_DISABLED_TAG SOCIAL_TAG "[socialdisabled]"
+#define SOCIAL_TEST_CASE(x, ...) ONLINE_TEST_CASE(x, SOCIAL_TAG __VA_ARGS__)
+
+//SOCIAL_TEST_CASE("Verify that AcceptFriendInvite returns an error if call with an invalid local user account id", EG_SOCIAL_ACCEPTFRIENDINVITE_TAG)
+//{
+//	const int32 NumUsersToLogin = 0;
+//
+//	FAcceptFriendInvite::Params OpAcceptFriendInviteParams;
+//	FAcceptFriendInviteHelper::FHelperParams AcceptInviteHelperParams;
+//	AcceptInviteHelperParams.OpParams = &OpAcceptFriendInviteParams;
+//	AcceptInviteHelperParams.OpParams->LocalAccountId = FAccountId();
+//	AcceptInviteHelperParams.ExpectedError = TOnlineResult<FAcceptFriendInvite>(Errors::InvalidParams());
+//
+//	GetLoginPipeline(NumUsersToLogin)
+//		.EmplaceStep<FAcceptFriendInviteHelper>(MoveTemp(AcceptInviteHelperParams));
+//
+//	RunToCompletion();
+//}
+
+SOCIAL_TEST_CASE("Verify that AcceptFriendInvite returns an error if call with an target user account id", EG_SOCIAL_ACCEPTFRIENDINVITE_TAG)
+{
+	FAccountId AccountId;
+
+	FAcceptFriendInvite::Params OpAcceptFriendInviteParams;
+	FAcceptFriendInviteHelper::FHelperParams AcceptInviteHelperParams;
+	AcceptInviteHelperParams.OpParams = &OpAcceptFriendInviteParams;
+	AcceptInviteHelperParams.OpParams->TargetAccountId = FAccountId();
+	AcceptInviteHelperParams.ExpectedError = TOnlineResult<FAcceptFriendInvite>(Errors::InvalidParams());
+
+	FTestPipeline& LoginPipeline = GetLoginPipeline(AccountId);
+
+	AcceptInviteHelperParams.OpParams->LocalAccountId = AccountId;
+
+	LoginPipeline
+		.EmplaceStep<FAcceptFriendInviteHelper>(MoveTemp(AcceptInviteHelperParams));
+
+	RunToCompletion();
+}
+
+//SOCIAL_TEST_CASE("Verify that AcceptFriendInvite returns a fail message if the local user is not logged in", EG_SOCIAL_ACCEPTFRIENDINVITEEOS_TAG)
+//{
+//	FAccountId FirstAccountId, SecondAccountId;
+//	int32 UserNumToLogout = 1;
+//	bool bLogout = true;
+//
+//	FAcceptFriendInvite::Params OpAcceptFriendInviteParams;
+//	FAcceptFriendInviteHelper::FHelperParams AcceptInviteHelperParams;
+//	AcceptInviteHelperParams.OpParams = &OpAcceptFriendInviteParams;
+//	AcceptInviteHelperParams.ExpectedError = TOnlineResult<FAcceptFriendInvite>(Errors::NotLoggedIn());
+//
+//	FTestPipeline& LoginPipeline = GetLoginPipeline(FirstAccountId, SecondAccountId);
+//
+//	AcceptInviteHelperParams.OpParams->LocalAccountId = FirstAccountId;
+//	AcceptInviteHelperParams.OpParams->TargetAccountId = SecondAccountId;
+//
+//	LoginPipeline
+//		.EmplaceStep<FAuthLogoutStep>(FPlatformMisc::GetPlatformUserForUserIndex(0))
+//		.EmplaceStep<FAcceptFriendInviteHelper>(MoveTemp(AcceptInviteHelperParams));
+//
+//	RunToCompletion(bLogout, UserNumToLogout);
+//}
+
+//SOCIAL_TEST_CASE("Verify that AcceptFriendInvite returns a fail message if ERelationship with target user is Friend")
+//{
+//	// TODO
+//}
+
+//SOCIAL_TEST_CASE("Verify that AcceptFriendInvite returns a fail message if ERelationship with target user is NotFriend")
+//{
+//	// TODO
+//}
+
+//SOCIAL_TEST_CASE("Verify that AcceptFriendInvite returns a fail message if ERelationship with target user is InviteSent")
+//{
+//	// TODO
+//}
+
+//SOCIAL_TEST_CASE("Verify that AcceptFriendInvite successfully completes if ERelationship with target user is InviteReceived, ERelationship becomes Friend")
+//{
+//	// TODO
+//}
+
+//SOCIAL_TEST_CASE("Verify that AcceptFriendInvite returns a fail message if ERelationship with target user is Blocked", EG_SOCIAL_DISABLED_TAG)
+//{
+//	// TODO
+//}
diff --git a/Engine/Source/Programs/Online/OnlineTestsCore/Source/Private/Tests/Social/BlockUserTests.cpp b/Engine/Source/Programs/Online/OnlineTestsCore/Source/Private/Tests/Social/BlockUserTests.cpp
new file mode 100644
--- /dev/null
+++ b/Engine/Source/Programs/Online/OnlineTestsCore/Source/Private/Tests/Social/BlockUserTests.cpp
@@ -0,0 +1,104 @@
+// Copyright Epic Games, Inc. All Rights Reserved.
+
+#include "Helpers/Social/BlockUserHelper.h"
+#include "Helpers/Auth/AuthLogout.h"
+
+#include "OnlineCatchHelper.h"
+
+#define SOCIAL_TAG "[suite_social]"
+#define EG_SOCIAL_BLOCKUSER_TAG SOCIAL_TAG "[blockuser]"
+#define EG_SOCIAL_BLOCKUSEREOS_TAG SOCIAL_TAG "[blockuser][.EOS]"
+#define SOCIAL_TEST_CASE(x, ...) ONLINE_TEST_CASE(x, SOCIAL_TAG __VA_ARGS__)
+
+SOCIAL_TEST_CASE("Verify that BlockUser returns a error if call with an invalid local user account id", EG_SOCIAL_BLOCKUSER_TAG)
+{
+	const int32 NumUsersToLogin = 0;
+
+	FBlockUser::Params OpBlockUserParams;
+	FBlockUserHelper::FHelperParams BlockUserHelperHelperParams;
+	BlockUserHelperHelperParams.OpParams = &OpBlockUserParams;
+	BlockUserHelperHelperParams.OpParams->LocalAccountId = FAccountId();
+
+	SubsystemType OnlineSubsystem = GetSubsystem();
+	EOnlineServices ServicesProvider = OnlineSubsystem->GetServicesProvider();
+	
+	/*if (ServicesProvider == EOnlineServices::Epic)
+	{
+		BlockUserHelperHelperParams.ExpectedError = TOnlineResult<FBlockUser>(Errors::InvalidParams());
+
+	}
+	else if (ServicesProvider == EOnlineServices::Xbox)
+	{
+		BlockUserHelperHelperParams.ExpectedError = TOnlineResult<FBlockUser>(Errors::InvalidUser());
+	}*/
+
+	BlockUserHelperHelperParams.ExpectedError = TOnlineResult<FBlockUser>(Errors::InvalidParams());
+
+
+	GetLoginPipeline(NumUsersToLogin)
+		.EmplaceStep<FBlockUserHelper>(MoveTemp(BlockUserHelperHelperParams));
+
+	RunToCompletion();
+}
+
+SOCIAL_TEST_CASE("Verify that BlockUser returns a error if call with an target user account id", EG_SOCIAL_BLOCKUSER_TAG)
+{
+	FAccountId AccountId;
+
+	FBlockUser::Params OpBlockUserParams;
+	FBlockUserHelper::FHelperParams BlockUserHelperHelperParams;
+	BlockUserHelperHelperParams.OpParams = &OpBlockUserParams;
+	BlockUserHelperHelperParams.OpParams->TargetAccountId = FAccountId();
+	BlockUserHelperHelperParams.ExpectedError = TOnlineResult<FBlockUser>(Errors::InvalidParams());
+
+	FTestPipeline& LoginPipeline = GetLoginPipeline(AccountId);
+
+	BlockUserHelperHelperParams.OpParams->LocalAccountId = AccountId;
+
+	LoginPipeline
+		.EmplaceStep<FBlockUserHelper>(MoveTemp(BlockUserHelperHelperParams));
+
+	RunToCompletion();
+}
+
+SOCIAL_TEST_CASE("Verify that BlockUser returns a fail message if the local user is not logged in", EG_SOCIAL_BLOCKUSEREOS_TAG)
+{
+	FAccountId FirstAccountId, SecondAccountId;
+	int32 UserNumToLogout = 1;
+	bool bLogout = true;
+
+	FBlockUser::Params OpBlockUserParams;
+	FBlockUserHelper::FHelperParams BlockUserHelperHelperParams;
+	BlockUserHelperHelperParams.OpParams = &OpBlockUserParams;
+	BlockUserHelperHelperParams.ExpectedError = TOnlineResult<FBlockUser>(Errors::NotLoggedIn());
+
+	FTestPipeline& LoginPipeline = GetLoginPipeline(FirstAccountId, SecondAccountId);
+
+	BlockUserHelperHelperParams.OpParams->LocalAccountId = FirstAccountId;
+	BlockUserHelperHelperParams.OpParams->TargetAccountId = SecondAccountId;
+
+	LoginPipeline
+		.EmplaceStep<FAuthLogoutStep>(FPlatformMisc::GetPlatformUserForUserIndex(0))
+		.EmplaceStep<FBlockUserHelper>(MoveTemp(BlockUserHelperHelperParams));
+
+	RunToCompletion(bLogout, UserNumToLogout);
+}
+
+SOCIAL_TEST_CASE("Verify that BlockUser completes successfully if both users are logged in", EG_SOCIAL_BLOCKUSEREOS_TAG)
+{
+	FAccountId FirstAccountId, SecondAccountId;
+
+	FBlockUser::Params OpBlockUserParams;
+	FBlockUserHelper::FHelperParams BlockUserHelperHelperParams;
+	BlockUserHelperHelperParams.OpParams = &OpBlockUserParams;
+
+	FTestPipeline& LoginPipeline = GetLoginPipeline(FirstAccountId, SecondAccountId);
+
+	BlockUserHelperHelperParams.OpParams->LocalAccountId = FirstAccountId;
+	BlockUserHelperHelperParams.OpParams->TargetAccountId = SecondAccountId;
+
+	LoginPipeline
+		.EmplaceStep<FBlockUserHelper>(MoveTemp(BlockUserHelperHelperParams));
+
+	RunToCompletion();
+}
diff --git a/Engine/Source/Programs/Online/OnlineTestsCore/Source/Private/Tests/Social/GetBlockedUsersTests.cpp b/Engine/Source/Programs/Online/OnlineTestsCore/Source/Private/Tests/Social/GetBlockedUsersTests.cpp
new file mode 100644
--- /dev/null
+++ b/Engine/Source/Programs/Online/OnlineTestsCore/Source/Private/Tests/Social/GetBlockedUsersTests.cpp
@@ -0,0 +1,84 @@
+// Copyright Epic Games, Inc. All Rights Reserved.
+
+#include "Helpers/Social/GetBlockedUsersHelper.h"
+#include "Helpers/Auth/AuthLogout.h"
+
+#include "OnlineCatchHelper.h"
+
+#define SOCIAL_TAG "[suite_social]"
+#define EG_SOCIAL_GETBLOCKEDUSERSGDK_TAG SOCIAL_TAG "[getblockedusers][.GDK]"
+#define EG_SOCIAL_DISABLED_TAG SOCIAL_TAG "[socialdisabled]"
+#define SOCIAL_TEST_CASE(x, ...) ONLINE_TEST_CASE(x, SOCIAL_TAG __VA_ARGS__)
+
+SOCIAL_TEST_CASE("Verify that GetBlockedUsers returns a error if call with an invalid local user account id", EG_SOCIAL_GETBLOCKEDUSERSGDK_TAG)
+{
+	const int32 NumUsersToLogin = 0;
+	FGetBlockedUsers::Params OpGetBlockedUsersParams;
+	FGetBlockedUsersHelper::FHelperParams GetBlockedUsersHelperParams;
+	GetBlockedUsersHelperParams.OpParams = &OpGetBlockedUsersParams;
+	GetBlockedUsersHelperParams.OpParams->LocalAccountId = FAccountId();
+
+	SubsystemType OnlineSubsystem = GetSubsystem();
+	EOnlineServices ServicesProvider = OnlineSubsystem->GetServicesProvider();
+
+	if (ServicesProvider == EOnlineServices::Epic)
+	{
+		GetBlockedUsersHelperParams.ExpectedError = TOnlineResult<FGetBlockedUsers>(Errors::NotImplemented());
+
+	}
+	else if (ServicesProvider == EOnlineServices::Xbox)
+	{
+		GetBlockedUsersHelperParams.ExpectedError = TOnlineResult<FGetBlockedUsers>(Errors::InvalidUser());
+	}
+
+	GetLoginPipeline(NumUsersToLogin)
+		.EmplaceStep<FGetBlockedUsersHelper>(MoveTemp(GetBlockedUsersHelperParams));
+
+	RunToCompletion();
+}
+
+SOCIAL_TEST_CASE("Verify that GetBlockedUsers returns a fail message if there are no avoid users for local user", EG_SOCIAL_GETBLOCKEDUSERSGDK_TAG)
+{
+	FAccountId AccountId;
+
+	FGetBlockedUsers::Params OpGetBlockedUsersParams;
+	FGetBlockedUsersHelper::FHelperParams GetBlockedUsersHelperParams;
+	GetBlockedUsersHelperParams.OpParams = &OpGetBlockedUsersParams;
+
+	SubsystemType OnlineSubsystem = GetSubsystem();
+	EOnlineServices ServicesProvider = OnlineSubsystem->GetServicesProvider();
+
+	if (ServicesProvider == EOnlineServices::Epic)
+	{
+		GetBlockedUsersHelperParams.ExpectedError = TOnlineResult<FGetBlockedUsers>(Errors::NotImplemented());
+
+	}
+	else if (ServicesProvider == EOnlineServices::Xbox)
+	{
+		GetBlockedUsersHelperParams.ExpectedError = TOnlineResult<FGetBlockedUsers>(Errors::InvalidUser());
+	}
+
+	FTestPipeline& LoginPipeline = GetLoginPipeline(AccountId);
+
+	GetBlockedUsersHelperParams.OpParams->LocalAccountId = AccountId;
+
+	LoginPipeline
+		.EmplaceStep<FGetBlockedUsersHelper>(MoveTemp(GetBlockedUsersHelperParams));
+
+	RunToCompletion();
+}
+
+//SOCIAL_TEST_CASE("Verify that GetBlockedUsers returns an empty list if there are no cached blocked users")
+//{
+//	// TODO
+//}
+
+//SOCIAL_TEST_CASE("Verify that GetBlockedUsers returns a list of 1 blocked user if there is 1 cached bocked user")
+//{
+//	// TODO
+//}
+
+//SOCIAL_TEST_CASE("Verify that GetBlockedUsers returns a list of all cached blocked user if there are multiple cached blocked users")
+//{
+//	// TODO
+//}
diff --git a/Engine/Source/Programs/Online/OnlineTestsCore/Source/Private/Tests/Social/GetFriendsTests.cpp b/Engine/Source/Programs/Online/OnlineTestsCore/Source/Private/Tests/Social/GetFriendsTests.cpp
new file mode 100644
--- /dev/null
+++ b/Engine/Source/Programs/Online/OnlineTestsCore/Source/Private/Tests/Social/GetFriendsTests.cpp
@@ -0,0 +1,114 @@
+// Copyright Epic Games, Inc. All Rights Reserved.
+
+#include "Helpers/Social/QueryFriendsHelper.h"
+#include "Helpers/Social/GetFriendsHelper.h"
+#include "Helpers/Auth/AuthLogout.h"
+
+#include "OnlineCatchHelper.h"
+
+#define SOCIAL_TAG "[suite_social]"
+#define EG_SOCIAL_GETFRIENDSEOS_TAG SOCIAL_TAG "[getfriends][.EOS]"
+#define EG_SOCIAL_GETFRIENDS_TAG SOCIAL_TAG "[getfriends]"
+
+#define SOCIAL_TEST_CASE(x, ...) ONLINE_TEST_CASE(x, SOCIAL_TAG __VA_ARGS__)
+
+SOCIAL_TEST_CASE("Verify that GetFriends returns an empty list if there are no cached Friends", EG_SOCIAL_GETFRIENDSEOS_TAG)
+{
+	FAccountId AccountId;	
+	int32 UserNumToLogin = 5;
+	int32 UserNumToLogout = 5;
+
+	FGetFriends::Params OpGetParams;
+	FGetFriendsHelper::FHelperParams GetFriendsHelperParams;
+	GetFriendsHelperParams.OpParams = &OpGetParams;
+	GetFriendsHelperParams.ExpectedError = TOnlineResult<FGetFriends>(Errors::InvalidState());
+
+	FTestPipeline& LoginPipeline = GetLoginPipeline(UserNumToLogin, AccountId);
+
+	GetFriendsHelperParams.OpParams->LocalAccountId = AccountId;
+
+	bool bLogout = true;
+
+	LoginPipeline
+		.EmplaceStep<FGetFriendsHelper>(MoveTemp(GetFriendsHelperParams));
+
+	RunToCompletion(bLogout, UserNumToLogout);
+}
+
+SOCIAL_TEST_CASE("Verify that GetFriends returns a list of 1 Friend if there is 1 cached Friend", EG_SOCIAL_GETFRIENDSEOS_TAG)
+{
+	FAccountId AccountId;
+	int32 UserNumToLogin = 6;
+	int32 UserNumToLogout = 6;
+	bool bLogout = true;
+
+	FQueryFriends::Params OpQueryParams;
+	FQueryFriendsHelper::FHelperParams QueryFriendsHelperParams;
+	QueryFriendsHelperParams.OpParams = &OpQueryParams;
+
+	FGetFriends::Params OpGetParams;
+	FGetFriendsHelper::FHelperParams GetFriendsHelperParams;
+	GetFriendsHelperParams.OpParams = &OpGetParams;
+
+	FTestPipeline& LoginPipeline = GetLoginPipeline(UserNumToLogin, AccountId);
+
+	QueryFriendsHelperParams.OpParams->LocalAccountId = AccountId;
+	GetFriendsHelperParams.OpParams->LocalAccountId = AccountId;
+
+	int32 ExpectedFriendsNum = 1;
+
+	LoginPipeline
+		.EmplaceStep<FQueryFriendsHelper>(MoveTemp(QueryFriendsHelperParams))
+		.EmplaceStep<FGetFriendsHelper>(MoveTemp(GetFriendsHelperParams), ExpectedFriendsNum);
+
+	RunToCompletion(bLogout, UserNumToLogout);
+}
+
+SOCIAL_TEST_CASE("Verify that GetFriends returns a list of all cached Friends if there are multiple cached Friends", EG_SOCIAL_GETFRIENDSEOS_TAG)
+{
+	FAccountId AccountId;
+
+	FQueryFriends::Params OpQueryParams;
+	FQueryFriendsHelper::FHelperParams QueryFriendsHelperParams;
+	QueryFriendsHelperParams.OpParams = &OpQueryParams;
+
+	FGetFriends::Params OpGetParams;
+	FGetFriendsHelper::FHelperParams GetFriendsHelperParams;
+	GetFriendsHelperParams.OpParams = &OpGetParams;
+
+	FTestPipeline& LoginPipeline = GetLoginPipeline(AccountId);
+
+	QueryFriendsHelperParams.OpParams->LocalAccountId = AccountId;
+	GetFriendsHelperParams.OpParams->LocalAccountId = AccountId;
+
+	int32 ExpectedFriendsNum = 5;
+
+	LoginPipeline
+		.EmplaceStep<FQueryFriendsHelper>(MoveTemp(QueryFriendsHelperParams))
+		.EmplaceStep<FGetFriendsHelper>(MoveTemp(GetFriendsHelperParams), ExpectedFriendsNum);
+
+	RunToCompletion();
+}
+
+SOCIAL_TEST_CASE("Verify that GetFriend returns a fail message if there are no cached Friends", EG_SOCIAL_GETFRIENDSEOS_TAG)
+{
+	FAccountId AccountId;	
+	int32 UserNumToLogin = 5;
+	int32 UserNumToLogout = 5;
+
+	FGetFriends::Params OpGetParams;
+	FGetFriendsHelper::FHelperParams GetFriendsHelperParams;
+	GetFriendsHelperParams.OpParams = &OpGetParams;
+	GetFriendsHelperParams.ExpectedError = TOnlineResult<FGetFriends>(Errors::InvalidState());
+
+	FTestPipeline& LoginPipeline = GetLoginPipeline(UserNumToLogin, AccountId);
+
+	GetFriendsHelperParams.OpParams->LocalAccountId = AccountId;
+
+	bool bLogout = true;
+
+	LoginPipeline
+		.EmplaceStep<FGetFriendsHelper>(MoveTemp(GetFriendsHelperParams));
+
+	RunToCompletion(bLogout, UserNumToLogout);
+}
diff --git a/Engine/Source/Programs/Online/OnlineTestsCore/Source/Private/Tests/Social/QueryBlockedUsersTests.cpp b/Engine/Source/Programs/Online/OnlineTestsCore/Source/Private/Tests/Social/QueryBlockedUsersTests.cpp
new file mode 100644
--- /dev/null
+++ b/Engine/Source/Programs/Online/OnlineTestsCore/Source/Private/Tests/Social/QueryBlockedUsersTests.cpp
@@ -0,0 +1,59 @@
+// Copyright Epic Games, Inc. All Rights Reserved.
+
+#include "Helpers/Social/QueryBlockedUsersHelper.h"
+#include "Helpers/Auth/AuthLogout.h"
+
+#include "OnlineCatchHelper.h"
+
+#define SOCIAL_TAG "[suite_social]"
+#define EG_SOCIAL_QUERYBLOCKEDUSERSGDK_TAG SOCIAL_TAG "[queryblockedusers][.GDK]"
+#define EG_SOCIAL_DISABLED_TAG SOCIAL_TAG "[socialdisabled]"
+#define SOCIAL_TEST_CASE(x, ...) ONLINE_TEST_CASE(x, SOCIAL_TAG __VA_ARGS__)
+
+SOCIAL_TEST_CASE("Verify that QueryBlockedUsers returns a error if call with an invalid local user account id", EG_SOCIAL_QUERYBLOCKEDUSERSGDK_TAG)
+{
+	const int32 NumUsersToLogin = 0;
+	
+	FQueryBlockedUsers::Params OpQueryParams;
+	FQueryBlockedUsersHelper::FHelperParams QueryBlockedUsersParams;
+	QueryBlockedUsersParams.OpParams = &OpQueryParams;
+	QueryBlockedUsersParams.OpParams->LocalAccountId = FAccountId();
+	
+	SubsystemType OnlineSubsystem = GetSubsystem();
+	EOnlineServices ServicesProvider = OnlineSubsystem->GetServicesProvider();
+	
+	if (ServicesProvider == EOnlineServices::Epic)
+	{
+		QueryBlockedUsersParams.ExpectedError = TOnlineResult<FQueryBlockedUsers>(Errors::NotImplemented());
+	
+	}
+	else if (ServicesProvider == EOnlineServices::Xbox)
+	{
+		QueryBlockedUsersParams.ExpectedError = TOnlineResult<FQueryBlockedUsers>(Errors::InvalidUser());
+	}
+	
+	GetLoginPipeline(NumUsersToLogin)
+		.EmplaceStep<FQueryBlockedUsersHelper>(MoveTemp(QueryBlockedUsersParams));
+	
+	RunToCompletion();
+}
+
+//SOCIAL_TEST_CASE("Verify that QueryBlockedUsers returns a fail message if the local user is not logged in")
+//{
+//	// TODO
+//}
+
+//SOCIAL_TEST_CASE("Verify that QueryBlockedUsers caches no blocked users if no blocked users exist for this user")
+//{
+//	// TODO
+//}
+
+//SOCIAL_TEST_CASE("Verify that QueryBlockedUsers caches one blocked user if only one blocked user exists for this user")
+//{
+//	// TODO
+//}
+
+//SOCIAL_TEST_CASE("Verify that QueryBlockedUsers caches all blocked users if multiple blocked users exist for this user")
+//{
+//	// TODO
+//}
diff --git a/Engine/Source/Programs/Online/OnlineTestsCore/Source/Private/Tests/Social/QueryFriendsTests.cpp b/Engine/Source/Programs/Online/OnlineTestsCore/Source/Private/Tests/Social/QueryFriendsTests.cpp
new file mode 100644
--- /dev/null
+++ b/Engine/Source/Programs/Online/OnlineTestsCore/Source/Private/Tests/Social/QueryFriendsTests.cpp
@@ -0,0 +1,131 @@
+// Copyright Epic Games, Inc. All Rights Reserved.
+
+#include "Helpers/Social/QueryFriendsHelper.h"
+#include "Helpers/Social/GetFriendsHelper.h"
+#include "Helpers/Auth/AuthLogout.h"
+
+#include "OnlineCatchHelper.h"
+
+#define SOCIAL_TAG "[suite_social]"
+#define EG_SOCIAL_QUERYFRIENDS_TAG SOCIAL_TAG "[queryfriends]"
+#define EG_SOCIAL_QUERYFRIENDSEOS_TAG SOCIAL_TAG "[queryfriends][.EOS]"
+#define SOCIAL_TEST_CASE(x, ...) ONLINE_TEST_CASE(x, SOCIAL_TAG __VA_ARGS__)
+
+SOCIAL_TEST_CASE("Verify that QueryFriends returns a fail message if the local user is not logged in", EG_SOCIAL_QUERYFRIENDSEOS_TAG)
+{
+	FAccountId AccountId;
+
+	FQueryFriends::Params OpQueryParams;
+	FQueryFriendsHelper::FHelperParams QueryFriendsHelperParams;
+	QueryFriendsHelperParams.OpParams = &OpQueryParams;
+	QueryFriendsHelperParams.ExpectedError = TOnlineResult<FQueryFriends>(Errors::NotLoggedIn());
+
+	FTestPipeline& LoginPipeline = GetLoginPipeline(AccountId);
+
+	QueryFriendsHelperParams.OpParams->LocalAccountId = AccountId;
+
+	bool bLogout = false;
+
+	LoginPipeline
+		.EmplaceStep<FAuthLogoutStep>(FPlatformMisc::GetPlatformUserForUserIndex(0))
+		.EmplaceStep<FQueryFriendsHelper>(MoveTemp(QueryFriendsHelperParams));
+
+	RunToCompletion(bLogout);
+}
+
+SOCIAL_TEST_CASE("Verify that QueryFriends returns a error if call with an invalid account id", EG_SOCIAL_QUERYFRIENDS_TAG)
+{
+	const int32 NumUsersToLogin = 0;
+
+	FQueryFriends::Params OpQueryParams;
+	FQueryFriendsHelper::FHelperParams QueryFriendsHelperParams;
+	QueryFriendsHelperParams.OpParams = &OpQueryParams;
+	QueryFriendsHelperParams.OpParams->LocalAccountId = FAccountId();
+	QueryFriendsHelperParams.ExpectedError = TOnlineResult<FQueryFriends>(Errors::InvalidParams());
+
+	GetLoginPipeline(NumUsersToLogin)
+		.EmplaceStep<FQueryFriendsHelper>(MoveTemp(QueryFriendsHelperParams));
+
+	RunToCompletion();
+}
+
+SOCIAL_TEST_CASE("Verify that QueryFriends caches no Friends if no Friends exist for this user", EG_SOCIAL_QUERYFRIENDSEOS_TAG)
+{
+	FAccountId AccountId;
+	int32 UserNumToLogin = 5;
+	int32 UserNumToLogout = 5;
+	bool bLogout = true;
+
+	FQueryFriends::Params OpQueryParams;
+	FQueryFriendsHelper::FHelperParams QueryFriendsHelperParams;
+	QueryFriendsHelperParams.OpParams = &OpQueryParams;
+	
+	FGetFriends::Params OpGetParams;
+	FGetFriendsHelper::FHelperParams GetFriendsHelperParams;
+	GetFriendsHelperParams.OpParams = &OpGetParams;
+
+	FTestPipeline& LoginPipeline = GetLoginPipeline(UserNumToLogin, AccountId);
+
+	QueryFriendsHelperParams.OpParams->LocalAccountId = AccountId;
+	GetFriendsHelperParams.OpParams->LocalAccountId = AccountId;
+	
+	LoginPipeline
+		.EmplaceStep<FQueryFriendsHelper>(MoveTemp(QueryFriendsHelperParams))
+		.EmplaceStep<FGetFriendsHelper>(MoveTemp(GetFriendsHelperParams));
+
+	RunToCompletion(bLogout, UserNumToLogout);
+}
+
+SOCIAL_TEST_CASE("Verify that QueryFriends caches one Friend if only one Friend exists for this user", EG_SOCIAL_QUERYFRIENDSEOS_TAG)
+{
+	FAccountId AccountId;
+	int32 UserNumToLogin = 6;
+	int32 UserNumToLogout = 6;
+	bool bLogout = true;
+
+	FQueryFriends::Params OpQueryParams;
+	FQueryFriendsHelper::FHelperParams QueryFriendsHelperParams;
+	QueryFriendsHelperParams.OpParams = &OpQueryParams;
+		
+	FGetFriends::Params OpGetParams;
+	FGetFriendsHelper::FHelperParams GetFriendsHelperParams;
+	GetFriendsHelperParams.OpParams = &OpGetParams;
+		
+	FTestPipeline& LoginPipeline = GetLoginPipeline(UserNumToLogin, AccountId);
+	
+	QueryFriendsHelperParams.OpParams->LocalAccountId = AccountId;
+	GetFriendsHelperParams.OpParams->LocalAccountId = AccountId;
+	
+	int32 ExpectedFriendsNum = 1;
+
+	LoginPipeline
+		.EmplaceStep<FQueryFriendsHelper>(MoveTemp(QueryFriendsHelperParams))
+		.EmplaceStep<FGetFriendsHelper>(MoveTemp(GetFriendsHelperParams), ExpectedFriendsNum);
+	
+	RunToCompletion(bLogout, UserNumToLogout);
+}
+
+SOCIAL_TEST_CASE("Verify that QueryFriends caches all Friends if multiple Friends exist for this user", EG_SOCIAL_QUERYFRIENDSEOS_TAG)
+{
+	FAccountId AccountId;
+	int32 ExpectedFriendsNum = 5;
+
+	FQueryFriends::Params OpQueryParams;
+	FQueryFriendsHelper::FHelperParams QueryFriendsHelperParams;
+	QueryFriendsHelperParams.OpParams = &OpQueryParams;
+
+	FGetFriends::Params OpGetParams;
+	FGetFriendsHelper::FHelperParams GetFriendsHelperParams;
+	GetFriendsHelperParams.OpParams = &OpGetParams;
+
+	FTestPipeline& LoginPipeline = GetLoginPipeline(AccountId);
+
+	QueryFriendsHelperParams.OpParams->LocalAccountId = AccountId;
+	GetFriendsHelperParams.OpParams->LocalAccountId = AccountId;
+
+	LoginPipeline
+		.EmplaceStep<FQueryFriendsHelper>(MoveTemp(QueryFriendsHelperParams))
+		.EmplaceStep<FGetFriendsHelper>(MoveTemp(GetFriendsHelperParams), ExpectedFriendsNum);
+
+	RunToCompletion();
+}
diff --git a/Engine/Source/Programs/Online/OnlineTestsCore/Source/Private/Tests/Social/RejectFriendInviteTests.cpp b/Engine/Source/Programs/Online/OnlineTestsCore/Source/Private/Tests/Social/RejectFriendInviteTests.cpp
new file mode 100644
--- /dev/null
+++ b/Engine/Source/Programs/Online/OnlineTestsCore/Source/Private/Tests/Social/RejectFriendInviteTests.cpp
@@ -0,0 +1,96 @@
+// Copyright Epic Games, Inc. All Rights Reserved.
+	
+#include "Helpers/Social/RejectFriendInviteHelper.h"
+#include "Helpers/Auth/AuthLogout.h"
+
+#include "OnlineCatchHelper.h"
+
+#define SOCIAL_TAG "[suite_social]"
+#define EG_SOCIAL_REJECTFRIENDINVITE_TAG SOCIAL_TAG "[rejectfriendinvite]"
+#define EG_SOCIAL_REJECTFRIENDINVITEEOS_TAG SOCIAL_TAG "[queryfriends][.EOS]"
+#define EG_SOCIAL_DISABLED_TAG SOCIAL_TAG "[socialdisabled]"
+#define SOCIAL_TEST_CASE(x, ...) ONLINE_TEST_CASE(x, SOCIAL_TAG __VA_ARGS__)
+
+SOCIAL_TEST_CASE("Verify that RejectFriendInvite returns a error if call with an invalid local user account id", EG_SOCIAL_REJECTFRIENDINVITE_TAG)
+{
+	const int32 NumUsersToLogin = 0;
+
+	FRejectFriendInvite::Params OpRejectFriendInviteParams;
+	FRejectFriendInviteHelper::FHelperParams RejectFriendInviteHelperParams;
+	RejectFriendInviteHelperParams.OpParams = &OpRejectFriendInviteParams;
+	RejectFriendInviteHelperParams.OpParams->LocalAccountId = FAccountId();
+	RejectFriendInviteHelperParams.ExpectedError = TOnlineResult<FRejectFriendInvite>(Errors::InvalidParams());
+
+	GetLoginPipeline(NumUsersToLogin)
+		.EmplaceStep<FRejectFriendInviteHelper>(MoveTemp(RejectFriendInviteHelperParams));
+
+	RunToCompletion();
+}
+
+SOCIAL_TEST_CASE("Verify that RejectFriendInvite returns a error if call with an target user account id", EG_SOCIAL_REJECTFRIENDINVITE_TAG)
+{
+	FAccountId AccountId;
+
+	FRejectFriendInvite::Params OpRejectFriendInviteParams;
+	FRejectFriendInviteHelper::FHelperParams RejectFriendInviteHelperParams;
+	RejectFriendInviteHelperParams.OpParams = &OpRejectFriendInviteParams;
+	RejectFriendInviteHelperParams.OpParams->TargetAccountId = FAccountId();
+	RejectFriendInviteHelperParams.ExpectedError = TOnlineResult<FRejectFriendInvite>(Errors::InvalidParams());
+
+	FTestPipeline& LoginPipeline = GetLoginPipeline(AccountId);
+
+	RejectFriendInviteHelperParams.OpParams->LocalAccountId = AccountId;
+
+	LoginPipeline
+		.EmplaceStep<FRejectFriendInviteHelper>(MoveTemp(RejectFriendInviteHelperParams));
+
+	RunToCompletion();
+}
+
+SOCIAL_TEST_CASE("Verify that RejectFriendInvite returns a fail message if the local user is not logged in", EG_SOCIAL_REJECTFRIENDINVITEEOS_TAG)
+{
+	FAccountId FirstAccountId, SecondAccountId;
+	int32 UserNumToLogout = 1;
+	bool bLogout = true;
+
+	FRejectFriendInvite::Params OpRejectFriendInviteParams;
+	FRejectFriendInviteHelper::FHelperParams RejectFriendInviteHelperParams;
+	RejectFriendInviteHelperParams.OpParams = &OpRejectFriendInviteParams;
+	RejectFriendInviteHelperParams.ExpectedError = TOnlineResult<FRejectFriendInvite>(Errors::NotLoggedIn());
+
+	FTestPipeline& LoginPipeline = GetLoginPipeline(FirstAccountId, SecondAccountId);
+
+	RejectFriendInviteHelperParams.OpParams->LocalAccountId = FirstAccountId;
+	RejectFriendInviteHelperParams.OpParams->TargetAccountId = SecondAccountId;
+
+	LoginPipeline
+		.EmplaceStep<FAuthLogoutStep>(FPlatformMisc::GetPlatformUserForUserIndex(0))
+		.EmplaceStep<FRejectFriendInviteHelper>(MoveTemp(RejectFriendInviteHelperParams));
+
+	RunToCompletion(bLogout, UserNumToLogout);
+}
+
+//SOCIAL_TEST_CASE("Verify that RejectFriendInvite returns a fail message if ERelationship with target user is Friend")
+//{
+//	// TODO
+//}
+
+//SOCIAL_TEST_CASE("Verify that RejectFriendInvite returns a fail message if ERelationship with target user is NotFriend")
+//{
+//	// TODO
+//}
+
+//SOCIAL_TEST_CASE("Verify that RejectFriendInvite returns a fail message if ERelationship with target user is InviteSent")
+//{
+//	// TODO
+//}
+
+//SOCIAL_TEST_CASE("Verify that RejectFriendInvite successfully completes if ERelationship with target user is InviteReceived, ERelationship becomes NotFriend")
+//{
+//	// TODO
+//}
+
+//SOCIAL_TEST_CASE("Verify that RejectFriendInvite returns a fail message if ERelationship with target user is Blocked")
+//{
+//	// TODO
+//}
diff --git a/Engine/Source/Programs/Online/OnlineTestsCore/Source/Private/Tests/Social/SendFriendInviteTests.cpp b/Engine/Source/Programs/Online/OnlineTestsCore/Source/Private/Tests/Social/SendFriendInviteTests.cpp
new file mode 100644
--- /dev/null
+++ b/Engine/Source/Programs/Online/OnlineTestsCore/Source/Private/Tests/Social/SendFriendInviteTests.cpp
@@ -0,0 +1,100 @@
+// Copyright Epic Games, Inc. All Rights Reserved.
+
+#include "Helpers/Social/GetFriendsHelper.h"	
+#include "Helpers/Social/QueryFriendsHelper.h"
+#include "Helpers/Social/SendFriendInviteHelper.h"
+#include "Helpers/Identity/IdentityGetLoginByUserId.h"
+#include "Helpers/Auth/AuthLogout.h"
+#include "Helpers/Auth/AuthLogin.h"
+
+#include "OnlineCatchHelper.h"
+
+#define SOCIAL_TAG "[suite_social]"
+#define EG_SOCIAL_SENDFRIENDINVITE_TAG SOCIAL_TAG "[sendfriendinvite]"
+#define EG_SOCIAL_SENDFRIENDINVITEEOS_TAG SOCIAL_TAG "[sendfriendinvite][.EOS]"
+#define EG_SOCIAL_DISABLED_TAG SOCIAL_TAG "[socialdisabled]"
+#define SOCIAL_TEST_CASE(x, ...) ONLINE_TEST_CASE(x, SOCIAL_TAG __VA_ARGS__)
+
+SOCIAL_TEST_CASE("Verify that SendFriendInvite returns a fail message if use invalid local user account id", EG_SOCIAL_SENDFRIENDINVITE_TAG)
+{
+	const int32 NumUsersToLogin = 0;
+
+	FSendFriendInvite::Params OpSendFriendInviteParams;
+	FSendFriendInviteHelper::FHelperParams SendFriendInviteHelperParams;
+	SendFriendInviteHelperParams.OpParams = &OpSendFriendInviteParams;
+	SendFriendInviteHelperParams.OpParams->LocalAccountId = FAccountId();
+	SendFriendInviteHelperParams.ExpectedError = TOnlineResult<FSendFriendInvite>(Errors::InvalidParams());
+
+	GetLoginPipeline(NumUsersToLogin)
+		.EmplaceStep<FSendFriendInviteHelper>(MoveTemp(SendFriendInviteHelperParams));
+
+	RunToCompletion();
+}
+
+SOCIAL_TEST_CASE("Verify that SendFriendInvite returns a fail message if use invalid target user account id", EG_SOCIAL_SENDFRIENDINVITE_TAG)
+{
+	FAccountId AccountId;
+
+	FSendFriendInvite::Params OpSendFriendInviteParams;
+	FSendFriendInviteHelper::FHelperParams SendFriendInviteHelperParams;
+	SendFriendInviteHelperParams.OpParams = &OpSendFriendInviteParams;
+	SendFriendInviteHelperParams.OpParams->TargetAccountId = FAccountId();
+	SendFriendInviteHelperParams.ExpectedError = TOnlineResult<FSendFriendInvite>(Errors::InvalidParams());
+
+	FTestPipeline& LoginPipeline = GetLoginPipeline(AccountId);
+
+	SendFriendInviteHelperParams.OpParams->LocalAccountId = AccountId;
+
+	LoginPipeline
+		.EmplaceStep<FSendFriendInviteHelper>(MoveTemp(SendFriendInviteHelperParams));
+
+	RunToCompletion();
+}
+
+SOCIAL_TEST_CASE("Verify that SendFriendInvite returns a fail message if the local user is not logged in", EG_SOCIAL_SENDFRIENDINVITEEOS_TAG)
+{
+	FAccountId FirstAccountId, SecondAccountId;
+
+	FSendFriendInvite::Params OpSendFriendInviteParams;
+	FSendFriendInviteHelper::FHelperParams SendFriendInviteHelperParams;
+	SendFriendInviteHelperParams.OpParams = &OpSendFriendInviteParams;
+	SendFriendInviteHelperParams.ExpectedError = TOnlineResult<FSendFriendInvite>(Errors::NotLoggedIn());
+
+	FTestPipeline& LoginPipeline = GetLoginPipeline(FirstAccountId, SecondAccountId);
+
+	SendFriendInviteHelperParams.OpParams->LocalAccountId = FirstAccountId;
+	SendFriendInviteHelperParams.OpParams->TargetAccountId = SecondAccountId;
+
+	bool bLogout = false;
+
+	LoginPipeline
+		.EmplaceStep<FAuthLogoutStep>(FPlatformMisc::GetPlatformUserForUserIndex(0))
+		.EmplaceStep<FSendFriendInviteHelper>(MoveTemp(SendFriendInviteHelperParams));
+
+	RunToCompletion(bLogout);
+}
+
+//SOCIAL_TEST_CASE("Verify that SendFriendInvite returns fail message if ERelationship with target user is Friend")
+//{
+//	// TODO
+//}
+
+//SOCIAL_TEST_CASE("Verify that SendFriendInvite completes successfully if ERelationship with target user is NotFriend, ERelationship becomes InviteSent")
+//{
+//	// TODO
+//}
+
+//SOCIAL_TEST_CASE("Verify that SendFriendInvite completes successfully if ERelationship with target user is InviteSent, ERelationship remains InviteSent")
+//{
+//	// TODO
+//}
+
+//SOCIAL_TEST_CASE("Verify that SendFriendInvite completes successfully if ERelationship with target user is InviteReceived, ERelationship becomes InviteSent")
+//{
+//	// TODO
+//}
+
+//SOCIAL_TEST_CASE("Verify that SendFriendInvite returns fail message if ERelationship with target user is Blocked")
+//{
+//	// TODO
+//}
diff --git a/Engine/Source/Programs/Online/OnlineTestsCore/Source/Private/Tests/Social/SocialTests.cpp b/Engine/Source/Programs/Online/OnlineTestsCore/Source/Private/Tests/Social/SocialTests.cpp
deleted file mode 100644
diff --git a/Engine/Source/Programs/Online/OnlineTestsCore/Source/Public/Helpers/Social/AcceptFriendInviteHelper.h b/Engine/Source/Programs/Online/OnlineTestsCore/Source/Public/Helpers/Social/AcceptFriendInviteHelper.h
new file mode 100644
--- /dev/null
+++ b/Engine/Source/Programs/Online/OnlineTestsCore/Source/Public/Helpers/Social/AcceptFriendInviteHelper.h
@@ -0,0 +1,52 @@
+// Copyright Epic Games, Inc. All Rights Reserved.
+
+#pragma once
+
+#include "Online/SocialCommon.h"
+#include "AsyncTestStep.h"
+#include "OnlineCatchHelper.h"
+
+struct FAcceptFriendInviteHelper : public FAsyncTestStep
+{
+	using ParamsType = UE::Online::FAcceptFriendInvite::Params;
+	using ResultType = UE::Online::TOnlineResult<UE::Online::FAcceptFriendInvite>;
+
+	struct FHelperParams
+	{
+		ParamsType* OpParams = nullptr;
+		TOptional<ResultType> ExpectedError;
+	};
+
+	FAcceptFriendInviteHelper(FHelperParams&& InHelperParams)
+		: HelperParams(MoveTemp(InHelperParams))
+	{
+		REQUIRE(HelperParams.OpParams);
+		REQUIRE((!HelperParams.ExpectedError.IsSet() || HelperParams.ExpectedError->IsError()));
+	}
+
+	virtual ~FAcceptFriendInviteHelper() = default;
+
+	virtual void Run(FAsyncStepResult Promise, SubsystemType Services) override
+	{
+		SocialInterface = Services->GetSocialInterface();
+		REQUIRE(SocialInterface);
+
+		SocialInterface->AcceptFriendInvite(MoveTemp(*HelperParams.OpParams))
+			.OnComplete([this, Promise = MoveTemp(Promise)](const ResultType& Result)
+				{
+					if (!HelperParams.ExpectedError.IsSet())
+					{
+						REQUIRE(Result.IsOk());
+					}
+					else
+					{
+						REQUIRE_OP_EQ(Result, HelperParams.ExpectedError->GetErrorValue());
+					}
+					Promise->SetValue(true);
+				});
+	}
+
+protected:
+	FHelperParams HelperParams;
+	UE::Online::ISocialPtr SocialInterface = nullptr;
+};
\ No newline at end of file
diff --git a/Engine/Source/Programs/Online/OnlineTestsCore/Source/Public/Helpers/Social/BlockUserHelper.h b/Engine/Source/Programs/Online/OnlineTestsCore/Source/Public/Helpers/Social/BlockUserHelper.h
new file mode 100644
--- /dev/null
+++ b/Engine/Source/Programs/Online/OnlineTestsCore/Source/Public/Helpers/Social/BlockUserHelper.h
@@ -0,0 +1,52 @@
+// Copyright Epic Games, Inc. All Rights Reserved.
+
+#pragma once
+
+#include "Online/SocialCommon.h"
+#include "AsyncTestStep.h"
+#include "OnlineCatchHelper.h"
+
+struct FBlockUserHelper : public FAsyncTestStep
+{
+	using ParamsType = UE::Online::FBlockUser::Params;
+	using ResultType = UE::Online::TOnlineResult<UE::Online::FBlockUser>;
+
+	struct FHelperParams
+	{
+		ParamsType* OpParams = nullptr;
+		TOptional<ResultType> ExpectedError;
+	};
+
+	FBlockUserHelper(FHelperParams&& InHelperParams)
+		: HelperParams(MoveTemp(InHelperParams))
+	{
+		REQUIRE(HelperParams.OpParams);
+		REQUIRE((!HelperParams.ExpectedError.IsSet() || HelperParams.ExpectedError->IsError()));
+	}
+
+	virtual ~FBlockUserHelper() = default;
+
+	virtual void Run(FAsyncStepResult Promise, SubsystemType Services) override
+	{
+		SocialInterface = Services->GetSocialInterface();
+		REQUIRE(SocialInterface);
+
+		SocialInterface->BlockUser(MoveTemp(*HelperParams.OpParams))
+			.OnComplete([this, Promise = MoveTemp(Promise)](const ResultType& Result)
+				{
+					if (!HelperParams.ExpectedError.IsSet())
+					{
+						REQUIRE(Result.IsOk());
+					}
+					else
+					{
+						REQUIRE_OP_EQ(Result, HelperParams.ExpectedError->GetErrorValue());
+					}
+					Promise->SetValue(true);
+				});
+	}
+
+protected:
+	FHelperParams HelperParams;
+	UE::Online::ISocialPtr SocialInterface = nullptr;
+};
\ No newline at end of file
diff --git a/Engine/Source/Programs/Online/OnlineTestsCore/Source/Public/Helpers/Social/GetBlockedUsersHelper.h b/Engine/Source/Programs/Online/OnlineTestsCore/Source/Public/Helpers/Social/GetBlockedUsersHelper.h
new file mode 100644
--- /dev/null
+++ b/Engine/Source/Programs/Online/OnlineTestsCore/Source/Public/Helpers/Social/GetBlockedUsersHelper.h
@@ -0,0 +1,56 @@
+// Copyright Epic Games, Inc. All Rights Reserved.
+
+#pragma once
+
+#include "Online/SocialCommon.h"
+#include "AsyncTestStep.h"
+#include "OnlineCatchHelper.h"
+
+struct FGetBlockedUsersHelper : public FAsyncTestStep
+{
+	using ParamsType = UE::Online::FGetBlockedUsers::Params;
+	using ResultType = UE::Online::TOnlineResult<UE::Online::FGetBlockedUsers>;
+
+	struct FHelperParams
+	{
+		ParamsType* OpParams = nullptr;
+		TOptional<ResultType> ExpectedError;
+	};
+
+	FGetBlockedUsersHelper(FHelperParams&& InHelperParams, const TOptional<uint32_t> InExpectedBlockedUsersNum = TOptional<uint32_t>())
+		: HelperParams(MoveTemp(InHelperParams))
+		, ExpectedBlockedUsersNum(InExpectedBlockedUsersNum)
+	{
+		REQUIRE(HelperParams.OpParams);
+		REQUIRE((!HelperParams.ExpectedError.IsSet() || HelperParams.ExpectedError->IsError()));
+	}
+
+	virtual ~FGetBlockedUsersHelper() = default;
+
+	virtual void Run(FAsyncStepResult Promise, SubsystemType Services) override
+	{
+		SocialInterface = Services->GetSocialInterface();
+		REQUIRE(SocialInterface);
+
+		ResultType Result = SocialInterface->GetBlockedUsers(MoveTemp(*HelperParams.OpParams));
+
+		if (HelperParams.ExpectedError.IsSet())
+		{
+			REQUIRE_OP_EQ(Result, HelperParams.ExpectedError->GetErrorValue());
+		}
+		else if (ExpectedBlockedUsersNum.IsSet())
+		{
+			CHECK(Result.GetOkValue().BlockedUsers.Num() == ExpectedBlockedUsersNum.GetValue());
+		}
+		else
+		{
+			CHECK(Result.GetOkValue().BlockedUsers.IsEmpty());
+		}
+		Promise->SetValue(true);
+	}
+
+protected:
+	FHelperParams HelperParams;
+	TOptional<uint32_t> ExpectedBlockedUsersNum = TOptional<uint32_t>();
+	UE::Online::ISocialPtr SocialInterface = nullptr;
+};
\ No newline at end of file
diff --git a/Engine/Source/Programs/Online/OnlineTestsCore/Source/Public/Helpers/Social/GetFriendsHelper.h b/Engine/Source/Programs/Online/OnlineTestsCore/Source/Public/Helpers/Social/GetFriendsHelper.h
new file mode 100644
--- /dev/null
+++ b/Engine/Source/Programs/Online/OnlineTestsCore/Source/Public/Helpers/Social/GetFriendsHelper.h
@@ -0,0 +1,56 @@
+// Copyright Epic Games, Inc. All Rights Reserved.
+
+#pragma once
+
+#include "Online/SocialCommon.h"
+#include "AsyncTestStep.h"
+#include "OnlineCatchHelper.h"
+
+struct FGetFriendsHelper : public FAsyncTestStep
+{
+	using ParamsType = UE::Online::FGetFriends::Params;
+	using ResultType = UE::Online::TOnlineResult<UE::Online::FGetFriends>;
+
+	struct FHelperParams
+	{
+		ParamsType* OpParams = nullptr;
+		TOptional<ResultType> ExpectedError;
+	};
+
+	FGetFriendsHelper(FHelperParams&& InHelperParams, const TOptional<uint32_t> InExpectedFriendsNum = TOptional<uint32_t>())
+		: HelperParams(MoveTemp(InHelperParams))
+		, ExpectedFriendsNum(InExpectedFriendsNum)
+	{
+		REQUIRE(HelperParams.OpParams);
+		REQUIRE((!HelperParams.ExpectedError.IsSet() || HelperParams.ExpectedError->IsError()));
+	}
+
+	virtual ~FGetFriendsHelper() = default;
+
+	virtual void Run(FAsyncStepResult Promise, SubsystemType Services) override
+	{
+		SocialInterface = Services->GetSocialInterface();
+		REQUIRE(SocialInterface);
+
+		ResultType Result = SocialInterface->GetFriends(MoveTemp(*HelperParams.OpParams));
+		
+		if (HelperParams.ExpectedError.IsSet())
+		{
+			REQUIRE_OP_EQ(Result, HelperParams.ExpectedError->GetErrorValue());
+		}
+		else if(ExpectedFriendsNum.IsSet())
+		{
+			CHECK(Result.GetOkValue().Friends.Num() == ExpectedFriendsNum.GetValue());
+		}
+		else
+		{
+			CHECK(Result.GetOkValue().Friends.IsEmpty());
+		}
+		Promise->SetValue(true);
+	}
+
+protected:
+	FHelperParams HelperParams;
+	TOptional<uint32_t> ExpectedFriendsNum = TOptional<uint32_t>();
+	UE::Online::ISocialPtr SocialInterface = nullptr;
+};
\ No newline at end of file
diff --git a/Engine/Source/Programs/Online/OnlineTestsCore/Source/Public/Helpers/Social/QueryBlockedUsersHelper.h b/Engine/Source/Programs/Online/OnlineTestsCore/Source/Public/Helpers/Social/QueryBlockedUsersHelper.h
new file mode 100644
--- /dev/null
+++ b/Engine/Source/Programs/Online/OnlineTestsCore/Source/Public/Helpers/Social/QueryBlockedUsersHelper.h
@@ -0,0 +1,52 @@
+// Copyright Epic Games, Inc. All Rights Reserved.
+
+#pragma once
+
+#include "Online/SocialCommon.h"
+#include "AsyncTestStep.h"
+#include "OnlineCatchHelper.h"
+
+struct FQueryBlockedUsersHelper : public FAsyncTestStep
+{
+	using ParamsType = UE::Online::FQueryBlockedUsers::Params;
+	using ResultType = UE::Online::TOnlineResult<UE::Online::FQueryBlockedUsers>;
+
+	struct FHelperParams
+	{
+		ParamsType* OpParams = nullptr;
+		TOptional<ResultType> ExpectedError;
+	};
+
+	FQueryBlockedUsersHelper(FHelperParams&& InHelperParams)
+		: HelperParams(MoveTemp(InHelperParams))
+	{
+		REQUIRE(HelperParams.OpParams);
+		REQUIRE((!HelperParams.ExpectedError.IsSet() || HelperParams.ExpectedError->IsError()));
+	}
+
+	virtual ~FQueryBlockedUsersHelper() = default;
+
+	virtual void Run(FAsyncStepResult Promise, SubsystemType Services) override
+	{
+		SocialInterface = Services->GetSocialInterface();
+		REQUIRE(SocialInterface);
+
+		SocialInterface->QueryBlockedUsers(MoveTemp(*HelperParams.OpParams))
+			.OnComplete([this, Promise = MoveTemp(Promise)](const ResultType& Result)
+				{
+					if (!HelperParams.ExpectedError.IsSet())
+					{
+						REQUIRE(Result.IsOk());
+					}
+					else
+					{
+						REQUIRE_OP_EQ(Result, HelperParams.ExpectedError->GetErrorValue());
+					}
+					Promise->SetValue(true);
+				});
+	}
+
+protected:
+	FHelperParams HelperParams;
+	UE::Online::ISocialPtr SocialInterface = nullptr;
+};
\ No newline at end of file
diff --git a/Engine/Source/Programs/Online/OnlineTestsCore/Source/Public/Helpers/Social/QueryFriendsHelper.h b/Engine/Source/Programs/Online/OnlineTestsCore/Source/Public/Helpers/Social/QueryFriendsHelper.h
new file mode 100644
--- /dev/null
+++ b/Engine/Source/Programs/Online/OnlineTestsCore/Source/Public/Helpers/Social/QueryFriendsHelper.h
@@ -0,0 +1,52 @@
+// Copyright Epic Games, Inc. All Rights Reserved.
+
+#pragma once
+
+#include "Online/SocialCommon.h"
+#include "AsyncTestStep.h"
+#include "OnlineCatchHelper.h"
+
+struct FQueryFriendsHelper : public FAsyncTestStep
+{
+	using ParamsType = UE::Online::FQueryFriends::Params;
+	using ResultType = UE::Online::TOnlineResult<UE::Online::FQueryFriends>;
+
+	struct FHelperParams
+	{
+		ParamsType* OpParams = nullptr;
+		TOptional<ResultType> ExpectedError;
+	};
+
+	FQueryFriendsHelper(FHelperParams&& InHelperParams)
+		: HelperParams(MoveTemp(InHelperParams))
+	{
+		REQUIRE(HelperParams.OpParams);
+		REQUIRE((!HelperParams.ExpectedError.IsSet() || HelperParams.ExpectedError->IsError()));
+	}
+
+	virtual ~FQueryFriendsHelper() = default;
+
+	virtual void Run(FAsyncStepResult Promise, SubsystemType Services) override
+	{
+		SocialInterface = Services->GetSocialInterface();
+		REQUIRE(SocialInterface);
+
+		SocialInterface->QueryFriends(MoveTemp(*HelperParams.OpParams))
+			.OnComplete([this, Promise = MoveTemp(Promise)](const ResultType& Result)
+				{
+					if (!HelperParams.ExpectedError.IsSet())
+					{
+						REQUIRE(Result.IsOk());
+					}
+					else
+					{
+						REQUIRE_OP_EQ(Result, HelperParams.ExpectedError->GetErrorValue());
+					}
+					Promise->SetValue(true);
+				});
+	}
+
+protected:
+	FHelperParams HelperParams;
+	UE::Online::ISocialPtr SocialInterface = nullptr;
+};
\ No newline at end of file
diff --git a/Engine/Source/Programs/Online/OnlineTestsCore/Source/Public/Helpers/Social/RejectFriendInviteHelper.h b/Engine/Source/Programs/Online/OnlineTestsCore/Source/Public/Helpers/Social/RejectFriendInviteHelper.h
new file mode 100644
--- /dev/null
+++ b/Engine/Source/Programs/Online/OnlineTestsCore/Source/Public/Helpers/Social/RejectFriendInviteHelper.h
@@ -0,0 +1,52 @@
+// Copyright Epic Games, Inc. All Rights Reserved.
+
+#pragma once
+
+#include "Online/SocialCommon.h"
+#include "AsyncTestStep.h"
+#include "OnlineCatchHelper.h"
+
+struct FRejectFriendInviteHelper : public FAsyncTestStep
+{
+	using ParamsType = UE::Online::FRejectFriendInvite::Params;
+	using ResultType = UE::Online::TOnlineResult<UE::Online::FRejectFriendInvite>;
+
+	struct FHelperParams
+	{
+		ParamsType* OpParams = nullptr;
+		TOptional<ResultType> ExpectedError;
+	};
+
+	FRejectFriendInviteHelper(FHelperParams&& InHelperParams)
+		: HelperParams(MoveTemp(InHelperParams))
+	{
+		REQUIRE(HelperParams.OpParams);
+		REQUIRE((!HelperParams.ExpectedError.IsSet() || HelperParams.ExpectedError->IsError()));
+	}
+
+	virtual ~FRejectFriendInviteHelper() = default;
+
+	virtual void Run(FAsyncStepResult Promise, SubsystemType Services) override
+	{
+		SocialInterface = Services->GetSocialInterface();
+		REQUIRE(SocialInterface);
+
+		SocialInterface->RejectFriendInvite(MoveTemp(*HelperParams.OpParams))
+			.OnComplete([this, Promise = MoveTemp(Promise)](const ResultType& Result)
+				{
+					if (!HelperParams.ExpectedError.IsSet())
+					{
+						REQUIRE(Result.IsOk());
+					}
+					else
+					{
+						REQUIRE_OP_EQ(Result, HelperParams.ExpectedError->GetErrorValue());
+					}
+					Promise->SetValue(true);
+				});
+	}
+
+protected:
+	FHelperParams HelperParams;
+	UE::Online::ISocialPtr SocialInterface = nullptr;
+};
\ No newline at end of file
diff --git a/Engine/Source/Programs/Online/OnlineTestsCore/Source/Public/Helpers/Social/SendFriendInviteHelper.h b/Engine/Source/Programs/Online/OnlineTestsCore/Source/Public/Helpers/Social/SendFriendInviteHelper.h
new file mode 100644
--- /dev/null
+++ b/Engine/Source/Programs/Online/OnlineTestsCore/Source/Public/Helpers/Social/SendFriendInviteHelper.h
@@ -0,0 +1,52 @@
+// Copyright Epic Games, Inc. All Rights Reserved.
+
+#pragma once
+
+#include "Online/SocialCommon.h"
+#include "AsyncTestStep.h"
+#include "OnlineCatchHelper.h"
+
+struct FSendFriendInviteHelper : public FAsyncTestStep
+{
+	using ParamsType = UE::Online::FSendFriendInvite::Params;
+	using ResultType = UE::Online::TOnlineResult<UE::Online::FSendFriendInvite>;
+
+	struct FHelperParams
+	{
+		ParamsType* OpParams = nullptr;
+		TOptional<ResultType> ExpectedError;
+	};
+
+	FSendFriendInviteHelper(FHelperParams&& InHelperParams)
+		: HelperParams(MoveTemp(InHelperParams))
+	{
+		REQUIRE(HelperParams.OpParams);
+		REQUIRE((!HelperParams.ExpectedError.IsSet() || HelperParams.ExpectedError->IsError()));
+	}
+
+	virtual ~FSendFriendInviteHelper() = default;
+
+	virtual void Run(FAsyncStepResult Promise, SubsystemType Services) override
+	{
+		SocialInterface = Services->GetSocialInterface();
+		REQUIRE(SocialInterface);
+
+		SocialInterface->SendFriendInvite(MoveTemp(*HelperParams.OpParams))
+			.OnComplete([this, Promise = MoveTemp(Promise)](const ResultType& Result)
+				{
+					if (!HelperParams.ExpectedError.IsSet())
+					{
+						REQUIRE_OP(Result);
+					}
+					else
+					{
+						REQUIRE_OP_EQ(Result, HelperParams.ExpectedError->GetErrorValue());
+					}
+					Promise->SetValue(true);
+				});
+	}
+
+protected:
+	FHelperParams HelperParams;
+	UE::Online::ISocialPtr SocialInterface = nullptr;
+};
\ No newline at end of file
diff --git a/Engine/Source/Programs/Online/OnlineTestsCore/Source/Public/OnlineCatchHelper.h b/Engine/Source/Programs/Online/OnlineTestsCore/Source/Public/OnlineCatchHelper.h
--- a/Engine/Source/Programs/Online/OnlineTestsCore/Source/Public/OnlineCatchHelper.h
+++ b/Engine/Source/Programs/Online/OnlineTestsCore/Source/Public/OnlineCatchHelper.h
@@ -37,7 +37,7 @@ public:
 	static void UnloadServiceModules();
 
 	FTestPipeline& GetLoginPipeline(uint32 NumUsersToLogin = 1) const;
-
+	FTestPipeline& GetLoginPipeline(uint32 UserNumToLogin, FAccountId& OutAccountId) const;
 
 	typedef UE::Online::FAccountId FAccountId;
 public:
@@ -98,7 +98,7 @@ protected:
 	/* Returns the ini login category name for the configured service */
 	FString GetLoginCredentialCategory() const;
 
-	void RunToCompletion(bool bLogout = true) const;
+	void RunToCompletion(bool bLogout = true, const TOptional<int32> UserNumToLogout = TOptional<int32>()) const;
 
 	/* ITestInvoker */
 	virtual void invoke() const override = 0;
