Fixing the issue with not being able to skip Meta tests with +DisableTestTags in .ini file

#rnf

diff --git a/Engine/Source/Programs/Online/OnlineTestsCore/Source/Private/Tests/Meta/MetaTests.cpp b/Engine/Source/Programs/Online/OnlineTestsCore/Source/Private/Tests/Meta/MetaTests.cpp
--- a/Engine/Source/Programs/Online/OnlineTestsCore/Source/Private/Tests/Meta/MetaTests.cpp
+++ b/Engine/Source/Programs/Online/OnlineTestsCore/Source/Private/Tests/Meta/MetaTests.cpp
@@ -7,7 +7,7 @@
 #include "Online/MulticastAdapter.h"
 #include "OnlineCatchHelper.h"
 
-#define META_SUITE_TAGS "[.Meta][.Null]"
+#define META_SUITE_TAGS "[Meta][Null]"
 #define META_TEST_CASE(x, ...) ONLINE_TEST_CASE(x, META_SUITE_TAGS)
 
 using namespace UE::Online;
diff --git a/Engine/Source/Programs/Online/OnlineTestsCore/Source/Private/Tests/Meta/SyncLoginTest.cpp b/Engine/Source/Programs/Online/OnlineTestsCore/Source/Private/Tests/Meta/SyncLoginTest.cpp
--- a/Engine/Source/Programs/Online/OnlineTestsCore/Source/Private/Tests/Meta/SyncLoginTest.cpp
+++ b/Engine/Source/Programs/Online/OnlineTestsCore/Source/Private/Tests/Meta/SyncLoginTest.cpp
@@ -4,7 +4,7 @@
 #include <catch2/catch_test_macros.hpp>
 #include "OnlineCatchHelper.h"
 
-#define SYNCLOGIN_ERROR_TAG "[.NULL][Meta][SyncLogin]"
+#define SYNCLOGIN_ERROR_TAG "[NULL][Meta][SyncLogin]"
 #define SYNCLOGIN_ERROR_TEST_CASE(x, ...) ONLINE_TEST_CASE(x, SYNCLOGIN_ERROR_TAG __VA_ARGS__)
 
 SYNCLOGIN_ERROR_TEST_CASE("Confirm sync login is logging in before RunToCompletion is called")
