Regression testing to ensure FLobby::LocalName is set during Create and Join operations

diff --git a/Engine/Source/Programs/Online/OnlineTestsCore/Source/Private/Tests/Lobby/LobbyCreateInviteTests.cpp b/Engine/Source/Programs/Online/OnlineTestsCore/Source/Private/Tests/Lobby/LobbyCreateInviteTests.cpp
--- a/Engine/Source/Programs/Online/OnlineTestsCore/Source/Private/Tests/Lobby/LobbyCreateInviteTests.cpp
+++ b/Engine/Source/Programs/Online/OnlineTestsCore/Source/Private/Tests/Lobby/LobbyCreateInviteTests.cpp
@@ -6,27 +6,33 @@
 
 #include "Helpers/Identity/IdentityGetLoginByUserId.h"
 #include "Helpers/Lobby/LobbyCreateHelper.h"
+#include "Helpers/Lobby/LobbyJoinHelper.h"
 #include "OnlineCatchHelper.h"
 
 #define LOBBY_CREATE_INVITE_TAGS "[Lobby]"
 #define LOBBY_CREATE_INVITE_TEST_CASE(x, ...) ONLINE_TEST_CASE(x, LOBBY_CREATE_INVITE_TAGS __VA_ARGS__)
 
-LOBBY_CREATE_INVITE_TEST_CASE("Basic lobby create test")
+LOBBY_CREATE_INVITE_TEST_CASE("Basic lobby create and join test")
 {
 	int32 LocalUserNum = 0;
-	UE::Online::FCreateLobby::Params Params;
-	Params.LocalName = TEXT("TestLobby");
-	Params.SchemaId = TEXT("test");
-	Params.MaxMembers = 2;
-	Params.JoinPolicy = ELobbyJoinPolicy::PublicAdvertised;
-	Params.Attributes = { { TEXT("LobbyCreateTime"), (int64)10}};
+	UE::Online::FCreateLobby::Params LobbyCreateParams;
+	LobbyCreateParams.LocalName = TEXT("TestLobby");
+	LobbyCreateParams.SchemaId = TEXT("test");
+	LobbyCreateParams.MaxMembers = 2;
+	LobbyCreateParams.JoinPolicy = ELobbyJoinPolicy::PublicAdvertised;
+	LobbyCreateParams.Attributes = { { TEXT("LobbyCreateTime"), (int64)10}};
+
+	FJoinLobby::Params JoinLobbyParams;
+	JoinLobbyParams.LocalName = TEXT("TestLobby");
 
 	GetLoginPipeline(1)
 		.EmplaceStep<FIdentityGetUniquePlayerIdStep>(LocalUserNum, 
-			[&Params](UE::Online::FAccountInfo InAccountInfo) {
-				Params.LocalAccountId = InAccountInfo.AccountId;
+			[&LobbyCreateParams, &JoinLobbyParams](UE::Online::FAccountInfo InAccountInfo) {
+				LobbyCreateParams.LocalAccountId = InAccountInfo.AccountId;
+				JoinLobbyParams.LocalAccountId = InAccountInfo.AccountId;
 			})
-		.EmplaceStep<FLobbyCreateHelper>(&Params, true);
+		.EmplaceStep<FLobbyCreateHelper>(&LobbyCreateParams, [&JoinLobbyParams](UE::Online::FLobby InLobby) {
+				JoinLobbyParams.LobbyId = InLobby.LobbyId; }, true);
 
 	RunToCompletion();
 }
diff --git a/Engine/Source/Programs/Online/OnlineTestsCore/Source/Public/Helpers/Lobby/LobbyCreateHelper.h b/Engine/Source/Programs/Online/OnlineTestsCore/Source/Public/Helpers/Lobby/LobbyCreateHelper.h
--- a/Engine/Source/Programs/Online/OnlineTestsCore/Source/Public/Helpers/Lobby/LobbyCreateHelper.h
+++ b/Engine/Source/Programs/Online/OnlineTestsCore/Source/Public/Helpers/Lobby/LobbyCreateHelper.h
@@ -17,6 +17,13 @@ struct FLobbyCreateHelper : public FTestPipeline::FStep
 		, bShouldPass(bShouldPass)
 	{}
 
+	FLobbyCreateHelper(UE::Online::FCreateLobby::Params* Params, TFunction<void(UE::Online::FLobby)>&& LobbyGetter, const bool& bShouldPass = true)
+		: CreateParams(Params)
+		, LobbyGetter(MoveTemp(LobbyGetter))
+		, bShouldPass(bShouldPass)
+		
+	{}
+
 	virtual ~FLobbyCreateHelper() = default;
 
 	enum class EState { Init, CreateLobbyCalled, CreateLobbyInProgress, CreateLobbyComplete, Done } State = EState::Init;
@@ -44,6 +51,8 @@ struct FLobbyCreateHelper : public FTestPipeline::FStep
 				if (bShouldPass)
 				{
 					CHECK_OP(Result);
+					CHECK((*Result.GetOkValue().Lobby).LocalName == (*CreateParams).LocalName);
+					LobbyGetter(*Result.GetOkValue().Lobby);
 				}
 				else
 				{
@@ -75,4 +84,6 @@ protected:
 	bool bShouldPass;
 
 	UE::Online::ILobbiesPtr OnlineLobbiesPtr = nullptr;
+
+	TFunction<void(UE::Online::FLobby)> LobbyGetter;
 };
diff --git a/Engine/Source/Programs/Online/OnlineTestsCore/Source/Public/Helpers/Lobby/LobbyJoinHelper.h b/Engine/Source/Programs/Online/OnlineTestsCore/Source/Public/Helpers/Lobby/LobbyJoinHelper.h
new file mode 100644
--- /dev/null
+++ b/Engine/Source/Programs/Online/OnlineTestsCore/Source/Public/Helpers/Lobby/LobbyJoinHelper.h
@@ -0,0 +1,66 @@
+// Copyright Epic Games, Inc. All Rights Reserved.
+
+#pragma once
+
+#include "Online/SessionsCommon.h"
+#include "AsyncTestStep.h"
+#include "Online/Lobbies.h"
+#include "TestDriver.h"
+#include "TestHarness.h"
+
+#include "Online/AuthCommon.h"
+#include "Online/OnlineAsyncOp.h"
+#include "OnlineCatchHelper.h"
+
+
+struct FLobbyJoinHelper : public FAsyncTestStep
+{
+	using ParamsType = UE::Online::FJoinLobby::Params;
+	using ResultType = UE::Online::TOnlineResult<UE::Online::FJoinLobby>;
+
+	struct FHelperParams
+	{
+		ParamsType* OpParams;
+		TOptional<ResultType> ExpectedError;
+	};
+
+	FLobbyJoinHelper(UE::Online::FJoinLobby::Params* InHelperParams, const bool& bShouldPass = true)
+		: HelperParams(InHelperParams)
+		,bShouldPass(bShouldPass)
+	{
+		REQUIRE(HelperParams.OpParams);
+		REQUIRE((!HelperParams.ExpectedError.IsSet() || HelperParams.ExpectedError->IsError()));
+	}
+
+	virtual ~FLobbyJoinHelper() = default;
+
+	virtual void Run(FAsyncStepResult Promise, SubsystemType Services) override
+	{
+		LobbyInterface = Services->GetLobbiesInterface();
+		REQUIRE(LobbyInterface);
+
+		LobbyInterface->JoinLobby(MoveTemp(*HelperParams.OpParams))
+			.OnComplete([this, Promise = MoveTemp(Promise)](const ResultType& Result)
+				{
+					if (!HelperParams.ExpectedError.IsSet())
+					{
+						REQUIRE_OP(Result);
+					}
+					else
+					{
+						const UE::Online::FOnlineError* ErrorValue = Result.TryGetErrorValue();
+						REQUIRE(ErrorValue != nullptr);
+						REQUIRE(*ErrorValue == HelperParams.ExpectedError->GetErrorValue());
+					}
+
+					Promise->SetValue(true);
+				});
+	}
+
+protected:
+	FHelperParams HelperParams;
+	bool bShouldPass;
+
+	UE::Online::ILobbiesPtr LobbyInterface = nullptr;
+
+};
\ No newline at end of file
