UE-210683: OnlineTests build warning/error fix for the all platforms.

diff --git a/Engine/Source/Programs/Online/OnlineTestsCore/Source/Public/Helpers/Lobby/LobbyJoinHelper.h b/Engine/Source/Programs/Online/OnlineTestsCore/Source/Public/Helpers/Lobby/LobbyJoinHelper.h
--- a/Engine/Source/Programs/Online/OnlineTestsCore/Source/Public/Helpers/Lobby/LobbyJoinHelper.h
+++ b/Engine/Source/Programs/Online/OnlineTestsCore/Source/Public/Helpers/Lobby/LobbyJoinHelper.h
@@ -24,8 +24,8 @@ struct FLobbyJoinHelper : public FAsyncTestStep
 		TOptional<ResultType> ExpectedError;
 	};
 
-	FLobbyJoinHelper(UE::Online::FJoinLobby::Params* InHelperParams, const bool& bShouldPass = true)
-		: HelperParams(InHelperParams)
+	FLobbyJoinHelper(FHelperParams&& InHelperParams, const bool& bShouldPass = true)
+		: HelperParams(MoveTemp(InHelperParams))
 		,bShouldPass(bShouldPass)
 	{
 		REQUIRE(HelperParams.OpParams);
@@ -62,5 +62,4 @@ protected:
 	bool bShouldPass;
 
 	UE::Online::ILobbiesPtr LobbyInterface = nullptr;
-
 };
\ No newline at end of file
