OSSTests: enabling for PS4/PS5.

- Added .Build.cs and .Target.cs for platforms as well as .ini files for each and loading modules depending on which platform the target will run for.

diff --git a/Engine/Source/Programs/Online/OSSTestsCore/OnlineSubsystemCatchHelper.cpp b/Engine/Source/Programs/Online/OSSTestsCore/OnlineSubsystemCatchHelper.cpp
--- a/Engine/Source/Programs/Online/OSSTestsCore/OnlineSubsystemCatchHelper.cpp
+++ b/Engine/Source/Programs/Online/OSSTestsCore/OnlineSubsystemCatchHelper.cpp
@@ -67,7 +67,7 @@ FTestPipeline& OnlineSubsystemTestBase::GetLoginPipeline(uint32 NumUsersToLogin)
 	REQUIRE(NumLocalUsers == -1); // Don't call GetLoginPipeline more than once per test
 	NumLocalUsers = NumUsersToLogin;
 
-	bool bUseAutoLogin;
+	bool bUseAutoLogin = false;
 	FString LoginCredentialCategory = FString::Printf(TEXT("LoginCredentials %s"), *Subsystem);
 	GConfig->GetBool(*LoginCredentialCategory, TEXT("UseAutoLogin"), bUseAutoLogin, GEngineIni);
 
