OSSTests: module loading from .ini file.

- reading modules from .ini file for Win64/PS4/PS5.

diff --git a/Engine/Source/Programs/Online/OSSTestsCore/OnlineSubsystemCatchHelper.cpp b/Engine/Source/Programs/Online/OSSTestsCore/OnlineSubsystemCatchHelper.cpp
--- a/Engine/Source/Programs/Online/OSSTestsCore/OnlineSubsystemCatchHelper.cpp
+++ b/Engine/Source/Programs/Online/OSSTestsCore/OnlineSubsystemCatchHelper.cpp
@@ -11,7 +11,8 @@
 #include "Helpers/Identity/IdentityLogoutHelper.h"
 #include "OnlineSubsystemNames.h"
 #include "Misc/CommandLine.h"
-#include "Misc/ConfigCacheIni.h"
+
+#include "Online/CoreOnline.h"
 
 TArray<TFunction<void()>>* GetGlobalInitalizers()
 {
@@ -19,6 +20,86 @@ TArray<TFunction<void()>>* GetGlobalInitalizers()
 	return &gInitalizersToCallInMain;
 }
 
+TArray<OnlineSubsystemAutoReg::FApplicableServicesConfig> OnlineSubsystemAutoReg::GetApplicableServices()
+{
+	static TArray<FApplicableServicesConfig> ServicesConfig =
+		[]()
+		{
+			TArray<FApplicableServicesConfig> ServicesConfigInit;
+			if (const TCHAR* CmdLine = FCommandLine::Get())	
+			{
+				FString Values;
+				TArray<FString> ServicesTags;
+				if (FParse::Value(CmdLine, TEXT("-Services="), Values, false))
+				{
+					Values.ParseIntoArray(ServicesTags, TEXT(","));
+				}
+
+				if (ServicesTags.IsEmpty())
+				{
+					GConfig->GetArray(TEXT("OnlineServicesTests"), TEXT("DefaultServices"), ServicesTags, GEngineIni);
+				}
+
+				for (const FString& ServicesTag : ServicesTags)
+				{
+					FString ConfigCategory = FString::Printf(TEXT("OnlineServicesTests %s"), *ServicesTag);
+					FApplicableServicesConfig Config;
+					Config.Tag = ServicesTag;
+
+					FString ServicesType;
+					GConfig->GetString(*ConfigCategory, TEXT("ServicesType"), ServicesType, GEngineIni);
+					GConfig->GetArray(*ConfigCategory, TEXT("ModulesToLoad"), Config.ModulesToLoad, GEngineIni);
+
+					LexFromString(Config.ServicesType, *ServicesType);
+					if (Config.ServicesType != UE::Online::EOnlineServices::None)
+					{
+						ServicesConfigInit.Add(MoveTemp(Config));
+					}
+				}
+			}
+
+			return ServicesConfigInit;
+		}();
+
+		return ServicesConfig;
+}
+
+TArray<FString> GetServiceModules()
+{
+	TArray<FString> Modules;
+
+	for (const OnlineSubsystemAutoReg::FApplicableServicesConfig& Config : OnlineSubsystemAutoReg::GetApplicableServices())
+	{
+		for (const FString& Module : Config.ModulesToLoad)
+		{
+			Modules.AddUnique(Module);
+		}
+	}
+
+	return Modules;
+}
+
+void OnlineSubsystemTestBase::LoadServiceModules()
+{
+	for (const FString& Module : GetServiceModules())
+	{
+		FModuleManager::LoadModulePtr<IModuleInterface>(*Module);
+	}
+}
+
+void OnlineSubsystemTestBase::UnloadServiceModules()
+{
+	const TArray<FString>& Modules = GetServiceModules();
+	// Shutdown in reverse order
+	for (int Index = Modules.Num() - 1; Index >= 0; --Index)
+	{
+		if (IModuleInterface* Module = FModuleManager::Get().GetModule(*Modules[Index]))
+		{
+			Module->ShutdownModule();
+		}
+	}
+}
+
 void OnlineSubsystemTestBase::ConstructInternal(FString SubsystemName)
 {
 	Subsystem = SubsystemName;
diff --git a/Engine/Source/Programs/Online/OSSTestsCore/OnlineSubsystemCatchHelper.h b/Engine/Source/Programs/Online/OSSTestsCore/OnlineSubsystemCatchHelper.h
--- a/Engine/Source/Programs/Online/OSSTestsCore/OnlineSubsystemCatchHelper.h
+++ b/Engine/Source/Programs/Online/OSSTestsCore/OnlineSubsystemCatchHelper.h
@@ -5,6 +5,8 @@
 #include "TestHarness.h"
 #include "TestDriver.h"
 #include "Interfaces/OnlineIdentityInterface.h"
+#include "Misc/ConfigCacheIni.h"
+#include "Modules/ModuleManager.h"
 
 #include <catch2/catch_test_case_info.hpp>
 #include <catch2/interfaces/catch_interfaces_registry_hub.hpp>
@@ -14,6 +16,13 @@ class OnlineSubsystemTestBase : public Catch::ITestInvoker
 public:
 	void ConstructInternal(FString ServiceName);
 	virtual ~OnlineSubsystemTestBase();
+
+	/* Loads all necessary services for the current test run */
+	static void LoadServiceModules();
+
+	/* Unloads all necessary services for the current test run */
+	static void UnloadServiceModules();
+
 protected:
 
 	OnlineSubsystemTestBase();
@@ -51,6 +60,15 @@ public:
 		TArray<FString> ShouldFailTags;
 		TArray<FString> DisableTestTags;
 	};
+		
+	struct FApplicableServicesConfig
+	{
+		FString Tag;
+		UE::Online::EOnlineServices ServicesType;
+		TArray<FString> ModulesToLoad;
+	};
+
+	static TArray<FApplicableServicesConfig> GetApplicableServices();
 
 	/*
 	* Helper function that calls CheckAllTagsIsIn(const TArray<FString>&, const TArray<FString>&);
