OSS Tests: Fix issue creating Programs on platforms requiring registering input devices.

diff --git a/Engine/Source/Programs/Online/OSSTestsCore/OnlineSubsystemCatchHelper.cpp b/Engine/Source/Programs/Online/OSSTestsCore/OnlineSubsystemCatchHelper.cpp
--- a/Engine/Source/Programs/Online/OSSTestsCore/OnlineSubsystemCatchHelper.cpp
+++ b/Engine/Source/Programs/Online/OSSTestsCore/OnlineSubsystemCatchHelper.cpp
@@ -11,9 +11,30 @@
 #include "Helpers/Identity/IdentityLogoutHelper.h"
 #include "OnlineSubsystemNames.h"
 #include "Misc/CommandLine.h"
+#include "GenericPlatform/GenericPlatformInputDeviceMapper.h"
 
 #include "Online/CoreOnline.h"
 
+// Make sure there are registered input devices for N users and fire
+// OnInputDeviceConnectionChange delegate for interested online service code.
+void EnsureLocalUserCount(uint32 NumUsers)
+{
+	TArray<FPlatformUserId> Users;
+	IPlatformInputDeviceMapper::Get().GetAllActiveUsers(Users);
+
+	const uint32 PreviousUserCount = Users.Num();
+	const uint32 NewUserCount = NumUsers > PreviousUserCount ? NumUsers - PreviousUserCount : 0;
+
+	for (uint32 Index = 0; Index < NewUserCount; ++Index)
+	{
+		const uint32 NewUserIndex = PreviousUserCount + Index;
+		IPlatformInputDeviceMapper::Get().Internal_MapInputDeviceToUser(
+			FInputDeviceId::CreateFromInternalId(NewUserIndex),
+			FPlatformMisc::GetPlatformUserForUserIndex(NewUserIndex),
+			EInputDeviceConnectionState::Connected);
+	}
+}
+
 TArray<TFunction<void()>>* GetGlobalInitalizers()
 {
 	static TArray<TFunction<void()>> gInitalizersToCallInMain;
@@ -149,10 +170,19 @@ FTestPipeline& OnlineSubsystemTestBase::GetLoginPipeline(uint32 NumUsersToLogin)
 	NumLocalUsers = NumUsersToLogin;
 
 	bool bUseAutoLogin = false;
+	bool bUseImplicitLogin = false;
 	FString LoginCredentialCategory = FString::Printf(TEXT("LoginCredentials %s"), *Subsystem);
 	GConfig->GetBool(*LoginCredentialCategory, TEXT("UseAutoLogin"), bUseAutoLogin, GEngineIni);
+	GConfig->GetBool(*LoginCredentialCategory, TEXT("UseImplicitLogin"), bUseImplicitLogin, GEngineIni);
 
-	if (bUseAutoLogin)
+	// Make sure input delegates are fired for adding the required user count.
+	EnsureLocalUserCount(NumUsersToLogin);
+
+	if (bUseImplicitLogin)
+	{
+		// Users are expected to already be valid.
+	}
+	else if (bUseAutoLogin)
 	{
 		NumLocalUsers = 1;
 		Pipeline.EmplaceStep<FIdentityAutoLoginStep>(0);
@@ -175,11 +205,29 @@ FTestPipeline& OnlineSubsystemTestBase::GetPipeline()
 
 void OnlineSubsystemTestBase::RunToCompletion() const
 {
-	for (uint32 i = 0; i < NumLocalUsers; i++)
+	bool bUseAutoLogin = false;
+	bool bUseImplicitLogin = false;
+	FString LoginCredentialCategory = FString::Printf(TEXT("LoginCredentials %s"), *Subsystem);
+	GConfig->GetBool(*LoginCredentialCategory, TEXT("UseAutoLogin"), bUseAutoLogin, GEngineIni);
+	GConfig->GetBool(*LoginCredentialCategory, TEXT("UseImplicitLogin"), bUseImplicitLogin, GEngineIni);
+
+	if (bUseImplicitLogin)
 	{
-		Pipeline.EmplaceStep<FIdentityLogoutStep>(i);
+		// Users are expected to already be valid.
 	}
-
+	else if (bUseAutoLogin)
+	{
+		NumLocalUsers = 1;
+		Pipeline.EmplaceStep<FIdentityAutoLoginStep>(0);
+	}
+	else
+	{
+		for (uint32 i = 0; i < NumLocalUsers; i++)
+		{
+			Pipeline.EmplaceStep<FIdentityLogoutStep>(i);
+		}
+	}
+	
 	FName SubsystemName = FName(GetSubsystem());
 	CAPTURE(*GetSubsystem());
 	FPipelineTestContext TestContext = FPipelineTestContext(SubsystemName);
diff --git a/Engine/Source/Runtime/Core/Private/Misc/Paths.cpp b/Engine/Source/Runtime/Core/Private/Misc/Paths.cpp
--- a/Engine/Source/Runtime/Core/Private/Misc/Paths.cpp
+++ b/Engine/Source/Runtime/Core/Private/Misc/Paths.cpp
@@ -170,6 +170,7 @@ bool FPaths::IsStaged()
 			FileExists(Combine(EngineConfigDir(), FString::Printf(TEXT("StagedBuild_%s.ini"), FApp::GetProjectName())));
 
 		UE::Paths::bIsComputingStaged = false;
+		bHasInitialized = true;
 	}
 
 	return bIsStaged;
