Refactor AccountServiceTests for PS4 platform

diff --git a/Engine/Build/LowLevelTests.xml b/Engine/Build/LowLevelTests.xml
--- a/Engine/Build/LowLevelTests.xml
+++ b/Engine/Build/LowLevelTests.xml
@@ -133,7 +133,7 @@
 			<Expand Name="$(TestName)AfterSteps" Platform="$(Platform)" />	
 		</Node>
 		<Label Category="After" Name="$(ShortName) After $(Platform)" Requires="$(TestName) AfterStep $(Platform)"/>
-		<Property Name="AfterTestNodes" Value="$(AfterTestNodes);$(TestName) AfterStep $(Platform)" If="$(RunAllTests) Or $(Run$(TestName)Tests)"/>
+		<Property Name="AfterTestNodes" Value="$(AfterTestNodes);$(TestName) AfterStep $(Platform)" If="($(RunsAllTests) Or $(Run$(TestName)Tests)) And ($(TestsAllSupportedPlatforms) Or $(TestPlatform$(Platform)))"/>
 	</Macro>
 	
 	<Macro Name="BuildAndStageTest" Arguments="Platform;TestName;StagesWithProjectFile;TargetName;ExtraCompilationArgs">
@@ -229,7 +229,7 @@
 				</Do>
 			</Node>
 			<Label Category="Test" Name="$(ShortName) $(Platform)" Requires="$(ShortName) Tests $(Platform)"/>
-			<Property Name="TestNodes" Value="$(TestNodes);$(ShortName) Tests $(Platform)" If="$(RunAllTests) Or $(Run$(TestName)Tests)"/>
+			<Property Name="TestNodes" Value="$(TestNodes);$(ShortName) Tests $(Platform)" If="($(RunsAllTests) Or $(Run$(TestName)Tests)) And ($(TestsAllSupportedPlatforms) Or $(TestPlatform$(Platform)))"/>
 				
 			<Do If="'$(CheckedRunAfterStep)' == 'True' AND '$(HasAfterSteps)' == 'True'">
 				<Expand Name="TestAfterStep" TestName="$(TestName)" ShortName="$(ShortName)" Platform="$(Platform)"/>
diff --git a/Engine/Source/Programs/Online/OSSTestsCore/Helpers/TickForTime.h b/Engine/Source/Programs/Online/OSSTestsCore/Helpers/TickForTime.h
new file mode 100644
--- /dev/null
+++ b/Engine/Source/Programs/Online/OSSTestsCore/Helpers/TickForTime.h
@@ -0,0 +1,40 @@
+// Copyright Epic Games, Inc. All Rights Reserved.
+ 
+#pragma once
+ 
+#include "TestDriver.h"
+ 
+struct FTickForTime : public FTestPipeline::FStep
+{
+public:
+	FTickForTime(const FTimespan InTimeToWaitMilliseconds)
+		: TimeToWaitMilliseconds(InTimeToWaitMilliseconds)
+	{
+	}
+
+	virtual EContinuance Tick(IOnlineSubsystem* OnlineSubsystem) override
+	{
+		double Now = FPlatformTime::Seconds();
+
+		if (!bHasSetStartTime)
+		{
+			bHasSetStartTime = true;
+			StartTime = Now;
+			return EContinuance::ContinueStepping;
+		}
+
+		if (FTimespan::FromSeconds(Now - StartTime).GetTicks() >= TimeToWaitMilliseconds.GetTicks())
+		{
+			return EContinuance::Done;
+		}
+
+		return EContinuance::ContinueStepping;
+	}
+
+	virtual ~FTickForTime() = default;
+
+protected:
+	double StartTime;
+	bool bHasSetStartTime = false;
+	FTimespan TimeToWaitMilliseconds;
+};
diff --git a/Engine/Source/Programs/Online/OSSTestsCore/OnlineSubsystemCatchHelper.cpp b/Engine/Source/Programs/Online/OSSTestsCore/OnlineSubsystemCatchHelper.cpp
--- a/Engine/Source/Programs/Online/OSSTestsCore/OnlineSubsystemCatchHelper.cpp
+++ b/Engine/Source/Programs/Online/OSSTestsCore/OnlineSubsystemCatchHelper.cpp
@@ -198,7 +198,7 @@ FTestPipeline& OnlineSubsystemTestBase::GetLoginPipeline(uint32 NumUsersToLogin)
 	return Pipeline;
 }
 
-FTestPipeline& OnlineSubsystemTestBase::GetPipeline()
+FTestPipeline& OnlineSubsystemTestBase::GetPipeline() const
 {
 	return GetLoginPipeline(0);
 }
@@ -229,7 +229,6 @@ void OnlineSubsystemTestBase::RunToCompletion() const
 	}
 	
 	FName SubsystemName = FName(GetSubsystem());
-	CAPTURE(*GetSubsystem());
 	FPipelineTestContext TestContext = FPipelineTestContext(SubsystemName);
 	CHECK(Driver.AddPipeline(MoveTemp(Pipeline), TestContext));
 	REQUIRE(IOnlineSubsystem::IsEnabled(SubsystemName));
diff --git a/Engine/Source/Programs/Online/OSSTestsCore/OnlineSubsystemCatchHelper.h b/Engine/Source/Programs/Online/OSSTestsCore/OnlineSubsystemCatchHelper.h
--- a/Engine/Source/Programs/Online/OSSTestsCore/OnlineSubsystemCatchHelper.h
+++ b/Engine/Source/Programs/Online/OSSTestsCore/OnlineSubsystemCatchHelper.h
@@ -30,7 +30,7 @@ protected:
 	FString GetSubsystem() const;
 	FOnlineAccountCredentials GetCredentials(int LocalUserNum) const;
 	FTestPipeline& GetLoginPipeline(uint32 NumUsersToLogin = 1) const;
-	FTestPipeline& GetPipeline();
+	FTestPipeline& GetPipeline() const;
 	void RunToCompletion() const;
 
 	/* ITestInvoker */
@@ -118,3 +118,7 @@ void PREPROCESSOR_JOIN(OnlineSubsystemTest_,RegName)::invoke() const\
 
 #define ONLINESUBSYSTEM_TEST_CASE(Name, Tags, ...) \
 	INTERNAL_ONLINESUBSYSTEM_TEST_CASE_NAMED(INTERNAL_CATCH_UNIQUE_NAME(OnlineSubsystemRegistrar), Name, Tags, __VA_ARGS__)
+
+#define REQUIRE_OP(Op)\
+	CAPTURE(Op);\
+	REQUIRE(Op.WasSuccessful());
\ No newline at end of file
diff --git a/Engine/Source/Programs/Online/OSSTestsCore/TestDriver.h b/Engine/Source/Programs/Online/OSSTestsCore/TestDriver.h
--- a/Engine/Source/Programs/Online/OSSTestsCore/TestDriver.h
+++ b/Engine/Source/Programs/Online/OSSTestsCore/TestDriver.h
@@ -107,6 +107,23 @@ public:
 	};
 	using FStepPtr = TUniquePtr<FStep>;
 
+	class FLambdaStep : public FStep
+	{
+	public:
+		FLambdaStep(TUniqueFunction<void(IOnlineSubsystem*)>&& InLambda)
+			: Lambda(MoveTemp(InLambda))
+		{}
+
+		virtual EContinuance Tick(IOnlineSubsystem* Subsystem) override
+		{
+			Lambda(Subsystem);
+			return EContinuance::Done;
+		}
+
+	private:
+		TUniqueFunction<void(IOnlineSubsystem*)> Lambda;
+	};
+
 	FTestPipeline(FTestPipeline&&) = default;
 	FTestPipeline(const FTestPipeline&) = delete;
 
@@ -143,6 +160,11 @@ public:
 		return MoveTemp(*this);
 	}
 
+	FTestPipeline&& EmplaceLambda(TUniqueFunction<void(IOnlineSubsystem*)> Lambda)
+	{
+		return EmplaceStep<FLambdaStep>(MoveTemp(Lambda));
+	}
+
 	/** Generates a string suitable for INFO which will help identify where and when a specific failure has occurred. */
 	FString InfoString() const;
 
diff --git a/Engine/Source/Programs/Online/OSSTestsCore/Tests/AccountServiceTest.cpp b/Engine/Source/Programs/Online/OSSTestsCore/Tests/AccountServiceTest.cpp
--- a/Engine/Source/Programs/Online/OSSTestsCore/Tests/AccountServiceTest.cpp
+++ b/Engine/Source/Programs/Online/OSSTestsCore/Tests/AccountServiceTest.cpp
@@ -1,6 +1,6 @@
 // Copyright Epic Games, Inc. All Rights Reserved.
-#include "CoreMinimal.h"
-#include "TestHarness.h"
+
+#include "TestDriver.h"
 
 #include "Interfaces/OnlineIdentityInterface.h"
 #include "OnlineSubsystem.h"
@@ -8,6 +8,7 @@
 #include "OnlineSubsystemCatchHelper.h"
 
 #define ACCOUNTSERVICE_TAG "[AccountService]"
+
 #define ACCOUNTSERVICE_TEST_CASE(x, ...) ONLINESUBSYSTEM_TEST_CASE(x, ACCOUNTSERVICE_TAG __VA_ARGS__)
 
 ACCOUNTSERVICE_TEST_CASE("Verify if we can properly create a OnlineAccountCredentials object")
@@ -22,15 +23,20 @@ ACCOUNTSERVICE_TEST_CASE("Verify if we can properly create a OnlineAccountCreden
 
 ACCOUNTSERVICE_TEST_CASE("Verify if we can properly instantiate the OSS")
 {
-	IOnlineSubsystem* OnlineSubsystem = IOnlineSubsystem::Get(FName(GetSubsystem()));
-	IOnlineIdentityPtr OIPtr = OnlineSubsystem->GetIdentityInterface();
+	FTestPipeline& Pipeline = GetPipeline();
 
-	int32 LocalUserNum = 1;
-	FString LocalType = "test_account";
-	FString LocalId = "12345";
-	FString LocalToken = "fake_token";
-	FOnlineAccountCredentials LocalAccount = FOnlineAccountCredentials(LocalType, LocalId, LocalToken);
+	Pipeline.EmplaceLambda([](IOnlineSubsystem* Services)
+		{
+			int32 LocalUserNum = 0;
+			FString LocalType = "test_account";
+			FString LocalId = "12345";
+			FString LocalToken = "fake_token";
+			FOnlineAccountCredentials LocalAccount = FOnlineAccountCredentials(LocalType, LocalId, LocalToken);
+
+			IOnlineIdentityPtr IdentityInterface = Services->GetIdentityInterface();		
+			bool LoggedIn = IdentityInterface.Get()->Login(LocalUserNum, LocalAccount);
+			REQUIRE(LoggedIn);
+		});
 
-	bool LoggedIn = OIPtr.Get()->Login(LocalUserNum, LocalAccount);
-	CHECK(LoggedIn == true);
+	RunToCompletion();
 }
\ No newline at end of file
