Added Commerce tests for EOS QueryOffers, QueryOffersById, GetOffers and GetOffersById interfaces.

diff --git a/Engine/Source/Programs/Online/OnlineTestsCore/Source/Private/Tests/Commerce/TestCommerce.cpp b/Engine/Source/Programs/Online/OnlineTestsCore/Source/Private/Tests/Commerce/TestCommerce.cpp
new file mode 100644
--- /dev/null
+++ b/Engine/Source/Programs/Online/OnlineTestsCore/Source/Private/Tests/Commerce/TestCommerce.cpp
@@ -0,0 +1,971 @@
+// Copyright Epic Games, Inc. All Rights Reserved.
+
+#include "CoreMinimal.h"
+
+#include <catch2/catch_test_macros.hpp>
+
+#include "OnlineCatchHelper.h"
+#include "Helpers/Auth/AuthLogin.h"
+#include "Helpers/Auth/AuthLogout.h"
+#include "Helpers/Commerce/CommerceQueryOffersHelper.h"
+#include "Helpers/Commerce/CommerceQueryOffersByIdHelper.h"
+#include "Helpers/Commerce/CommerceGetOffersHelper.h"
+#include "Helpers/Commerce/CommerceGetOffersByIdHelper.h"
+#include "Helpers/Commerce/CommerceCheckoutHelper.h"
+#include "Helpers/Commerce/CommerceQueryEntitlementsHelper.h"
+#include "Helpers/Commerce/CommerceGetEntitlementsHelper.h"
+#include "Online/CommerceCommon.h"
+
+
+#define COMMERCE_TAG "[Commerce]"
+#define COMMERCE_QUERYOFFERS_TAG COMMERCE_TAG "[queryoffers]"
+#define COMMERCE_QUERYOFFERSBYID_TAG COMMERCE_TAG "[queryoffersbyid]"
+#define COMMERCE_GETOFFERS_TAG COMMERCE_TAG "[getoffers]"
+#define COMMERCE_GETOFFERSBYID_TAG COMMERCE_TAG "[getoffersbyid]"
+#define COMMERCE_DISABLED_TAG COMMERCE_TAG "[commercedisabled]"
+#define COMMERCE_TEST_CASE(x, ...) ONLINE_TEST_CASE(x, COMMERCE_TAG __VA_ARGS__)
+
+
+COMMERCE_TEST_CASE("Basic compile test")
+{
+	TSharedPtr<ICommerce> Commerce;
+	CHECK(!Commerce.IsValid());
+}
+
+// QueryOffers tests
+COMMERCE_TEST_CASE("Verify that QueryOffers returns a fail message if the local user is not logged in", COMMERCE_QUERYOFFERS_TAG)
+{
+	FAccountId AccountId;
+
+	FCommerceQueryOffers::Params OpQueryOffersParams;
+	FCommerceQueryOffersHelper::FHelperParams QueryOffersHelperParams;
+	QueryOffersHelperParams.OpParams = &OpQueryOffersParams;
+	QueryOffersHelperParams.ExpectedError = TOnlineResult<FCommerceQueryOffers>(Errors::NotLoggedIn());
+
+	FTestPipeline& LoginPipeline = GetLoginPipeline(AccountId);
+	QueryOffersHelperParams.OpParams->LocalAccountId = AccountId;
+
+	bool bLogout = false;
+
+	LoginPipeline
+		.EmplaceStep<FAuthLogoutStep>(FPlatformMisc::GetPlatformUserForUserIndex(0))
+		.EmplaceStep<FCommerceQueryOffersHelper>(MoveTemp(QueryOffersHelperParams));
+
+	RunToCompletion(bLogout);
+}
+
+COMMERCE_TEST_CASE("Verify that QueryOffers returns a fail message of the given local user ID does not match the actual local user ID", COMMERCE_QUERYOFFERS_TAG)
+{
+	FAccountId AccountId;
+
+	FCommerceQueryOffers::Params OpQueryOffersParams;
+	FCommerceQueryOffersHelper::FHelperParams QueryOffersHelperParams;
+	QueryOffersHelperParams.OpParams = &OpQueryOffersParams;
+	QueryOffersHelperParams.ExpectedError = TOnlineResult<FCommerceQueryOffers>(Errors::NotLoggedIn());
+
+	FTestPipeline& LoginPipeline = GetLoginPipeline(AccountId);
+	QueryOffersHelperParams.OpParams->LocalAccountId = FAccountId();
+
+	bool bLogout = true;
+
+	LoginPipeline
+		.EmplaceStep<FCommerceQueryOffersHelper>(MoveTemp(QueryOffersHelperParams));
+
+	RunToCompletion(bLogout);
+}
+
+COMMERCE_TEST_CASE("Verify that QueryOffers caches an empty list if there are no offers", COMMERCE_QUERYOFFERS_TAG)
+{
+	FAccountId AccountId;
+	TOptional<uint32_t> ExpectedOffersNum = 0;
+
+	FCommerceQueryOffers::Params OpQueryOffersParams;
+	FCommerceQueryOffersHelper::FHelperParams QueryOffersHelperParams;
+	QueryOffersHelperParams.OpParams = &OpQueryOffersParams;
+
+	FCommerceGetOffers::Params OpGetOffersParams;
+	FCommerceGetOffersHelper::FHelperParams GetOffersHelperParams;
+	GetOffersHelperParams.OpParams = &OpGetOffersParams;
+
+	FTestPipeline& LoginPipeline = GetLoginPipeline(AccountId);
+	QueryOffersHelperParams.OpParams->LocalAccountId = AccountId;
+	FString CatalogId;
+	GConfig->GetString(TEXT("EOSSDK.Platform.OnlineTests"), TEXT("OfferNamespace0Items"), CatalogId, GEngineIni);
+	QueryOffersHelperParams.OpParams->OverrideCatalogNamespace = CatalogId;
+	GetOffersHelperParams.OpParams->LocalAccountId = AccountId;
+
+	bool bLogout = true;
+
+	LoginPipeline
+		.EmplaceStep<FCommerceQueryOffersHelper>(MoveTemp(QueryOffersHelperParams))
+		.EmplaceStep<FCommerceGetOffersHelper>(MoveTemp(GetOffersHelperParams), ExpectedOffersNum);
+
+	RunToCompletion(bLogout);
+}
+
+COMMERCE_TEST_CASE("Verify that QueryOffers caches a list of one offer if there is only one existing offer", COMMERCE_QUERYOFFERS_TAG)
+{
+	FAccountId AccountId;
+	TOptional<uint32_t> ExpectedOffersNum = 1;
+
+	FCommerceQueryOffers::Params OpQueryOffersParams;
+	FCommerceQueryOffersHelper::FHelperParams QueryOffersHelperParams;
+	QueryOffersHelperParams.OpParams = &OpQueryOffersParams;
+
+	FCommerceGetOffers::Params OpGetOffersParams;
+	FCommerceGetOffersHelper::FHelperParams GetOffersHelperParams;
+	GetOffersHelperParams.OpParams = &OpGetOffersParams;
+
+	FTestPipeline& LoginPipeline = GetLoginPipeline(AccountId);
+	QueryOffersHelperParams.OpParams->LocalAccountId = AccountId;
+	FString CatalogId;
+	GConfig->GetString(TEXT("EOSSDK.Platform.OnlineTests"), TEXT("OfferNamespace1Items"), CatalogId, GEngineIni);
+	QueryOffersHelperParams.OpParams->OverrideCatalogNamespace = CatalogId;
+	GetOffersHelperParams.OpParams->LocalAccountId = AccountId;
+
+	bool bLogout = true;
+
+	LoginPipeline
+		.EmplaceStep<FCommerceQueryOffersHelper>(MoveTemp(QueryOffersHelperParams))
+		.EmplaceStep<FCommerceGetOffersHelper>(MoveTemp(GetOffersHelperParams), ExpectedOffersNum);
+
+	RunToCompletion(bLogout);
+}
+
+COMMERCE_TEST_CASE("Verify that QueryOffers caches the list of all offers if there are multiple existing offers", COMMERCE_QUERYOFFERS_TAG)
+{
+	FAccountId AccountId;
+	TOptional<uint32_t> ExpectedOffersNum = 4;
+
+	FCommerceQueryOffers::Params OpQueryOffersParams;
+	FCommerceQueryOffersHelper::FHelperParams QueryOffersHelperParams;
+	QueryOffersHelperParams.OpParams = &OpQueryOffersParams;
+
+	FCommerceGetOffers::Params OpGetOffersParams;
+	FCommerceGetOffersHelper::FHelperParams GetOffersHelperParams;
+	GetOffersHelperParams.OpParams = &OpGetOffersParams;
+
+	FTestPipeline& LoginPipeline = GetLoginPipeline(AccountId);
+	QueryOffersHelperParams.OpParams->LocalAccountId = AccountId;
+	GetOffersHelperParams.OpParams->LocalAccountId = AccountId;
+
+	bool bLogout = true;
+
+	LoginPipeline
+		.EmplaceStep<FCommerceQueryOffersHelper>(MoveTemp(QueryOffersHelperParams))
+		.EmplaceStep<FCommerceGetOffersHelper>(MoveTemp(GetOffersHelperParams), ExpectedOffersNum);
+
+	RunToCompletion(bLogout);
+}
+
+// QueryOffersById tests
+COMMERCE_TEST_CASE("Verify that QueryOffersById returns a fail message if the local user is not logged in", COMMERCE_QUERYOFFERSBYID_TAG)
+{
+	FAccountId AccountId;
+
+	FCommerceQueryOffersById::Params OpQueryOffersByIdParams;
+	FCommerceQueryOffersByIdHelper::FHelperParams QueryOffersByIdHelperParams;
+	QueryOffersByIdHelperParams.OpParams = &OpQueryOffersByIdParams;
+	QueryOffersByIdHelperParams.ExpectedError = TOnlineResult<FCommerceQueryOffersById>(Errors::NotLoggedIn());
+
+	FTestPipeline& LoginPipeline = GetLoginPipeline(AccountId);
+	QueryOffersByIdHelperParams.OpParams->LocalAccountId = AccountId;
+
+	bool bLogout = false;
+
+	LoginPipeline
+		.EmplaceStep<FAuthLogoutStep>(FPlatformMisc::GetPlatformUserForUserIndex(0))
+		.EmplaceStep<FCommerceQueryOffersByIdHelper>(MoveTemp(QueryOffersByIdHelperParams));
+
+	RunToCompletion(bLogout);
+}
+
+COMMERCE_TEST_CASE("Verify that QueryOffersById returns a fail message of the given local user ID does not match the actual local user ID", COMMERCE_QUERYOFFERSBYID_TAG)
+{
+	FAccountId AccountId;
+
+	FCommerceQueryOffersById::Params OpQueryOffersByIdParams;
+	FCommerceQueryOffersByIdHelper::FHelperParams QueryOffersByIdHelperParams;
+	QueryOffersByIdHelperParams.OpParams = &OpQueryOffersByIdParams;
+	QueryOffersByIdHelperParams.ExpectedError = TOnlineResult<FCommerceQueryOffersById>(Errors::NotLoggedIn());
+
+	FTestPipeline& LoginPipeline = GetLoginPipeline(AccountId);
+	QueryOffersByIdHelperParams.OpParams->LocalAccountId = FAccountId();
+
+	bool bLogout = true;
+
+	LoginPipeline
+		.EmplaceStep<FCommerceQueryOffersByIdHelper>(MoveTemp(QueryOffersByIdHelperParams));
+
+	RunToCompletion(bLogout);
+}
+
+COMMERCE_TEST_CASE("Verify that QueryOffersById returns an empty list when given an empty list of IDs and there are no existing offers", COMMERCE_DISABLED_TAG)
+{
+	// EOS QueryOffers interface doesn't support anything ID-specific, skipping the test
+}
+
+COMMERCE_TEST_CASE("Verify that QueryOffersById returns an empty list when given an empty list of IDs and there are multiple existing offers", COMMERCE_DISABLED_TAG)
+{
+	// EOS QueryOffers interface doesn't support anything ID-specific, skipping the test
+}
+
+COMMERCE_TEST_CASE("Verify that QueryOffersById returns a fail message when given a populated list of IDs and there are no existing offers", COMMERCE_DISABLED_TAG)
+{
+	// EOS QueryOffers interface doesn't support anything ID-specific, skipping the test
+}
+
+COMMERCE_TEST_CASE("Verify that QueryOffersById returns a fail message when given a populated list of IDs and none of them match existing offers", COMMERCE_DISABLED_TAG)
+{
+	// EOS QueryOffers interface doesn't support anything ID-specific, skipping the test
+}
+
+COMMERCE_TEST_CASE("Verify that QueryOffersById returns a fail message when given a populated list of IDs where one ID exists and another does not", COMMERCE_DISABLED_TAG)
+{
+	// EOS QueryOffers interface doesn't support anything ID-specific, skipping the test
+}
+
+COMMERCE_TEST_CASE("Verify that QueryOffersById returns the correct list of one offer when given the ID for one existing offer", COMMERCE_DISABLED_TAG)
+{
+	// EOS QueryOffers interface doesn't support anything ID-specific, skipping the test
+}
+
+COMMERCE_TEST_CASE("Verify that QueryOffersById returns the correct list of offers when given a populated list of multiple existing IDs", COMMERCE_DISABLED_TAG)
+{
+	// EOS QueryOffers interface doesn't support anything ID-specific, skipping the test
+}
+
+// GetOffers
+COMMERCE_TEST_CASE("Verify that GetOffers returns a fail message if the local user is not logged in", COMMERCE_GETOFFERS_TAG)
+{
+	FAccountId AccountId;
+
+	FCommerceGetOffers::Params OpGetOffersParams;
+	FCommerceGetOffersHelper::FHelperParams GetOffersHelperParams;
+	GetOffersHelperParams.OpParams = &OpGetOffersParams;
+	GetOffersHelperParams.ExpectedError = TOnlineResult<FCommerceGetOffers>(Errors::NotLoggedIn());
+
+	FTestPipeline& LoginPipeline = GetLoginPipeline(AccountId);
+	GetOffersHelperParams.OpParams->LocalAccountId = AccountId;
+
+	bool bLogout = false;
+
+	LoginPipeline
+		.EmplaceStep<FAuthLogoutStep>(FPlatformMisc::GetPlatformUserForUserIndex(0))
+		.EmplaceStep<FCommerceGetOffersHelper>(MoveTemp(GetOffersHelperParams));
+
+	RunToCompletion(bLogout);
+}
+
+COMMERCE_TEST_CASE("Verify that GetOffers returns a fail message of the given local user ID does not match the actual local user ID", COMMERCE_GETOFFERS_TAG)
+{
+	FAccountId AccountId;
+
+	FCommerceGetOffers::Params OpGetOffersParams;
+	FCommerceGetOffersHelper::FHelperParams GetOffersHelperParams;
+	GetOffersHelperParams.OpParams = &OpGetOffersParams;
+	GetOffersHelperParams.ExpectedError = TOnlineResult<FCommerceGetOffers>(Errors::NotLoggedIn());
+
+	FTestPipeline& LoginPipeline = GetLoginPipeline(AccountId);
+	GetOffersHelperParams.OpParams->LocalAccountId = FAccountId();
+
+	bool bLogout = true;
+
+	LoginPipeline
+		.EmplaceStep<FCommerceGetOffersHelper>(MoveTemp(GetOffersHelperParams));
+
+	RunToCompletion(bLogout);
+}
+
+COMMERCE_TEST_CASE("Verify that GetOffers returns an empty list if there are no cached offers", COMMERCE_GETOFFERS_TAG)
+{
+	FAccountId AccountId;
+	TOptional<uint32_t> ExpectedOffersNum = 0;
+
+	FCommerceQueryOffers::Params OpQueryOffersParams;
+	FCommerceQueryOffersHelper::FHelperParams QueryOffersHelperParams;
+	QueryOffersHelperParams.OpParams = &OpQueryOffersParams;
+
+	FCommerceGetOffers::Params OpGetOffersParams;
+	FCommerceGetOffersHelper::FHelperParams GetOffersHelperParams;
+	GetOffersHelperParams.OpParams = &OpGetOffersParams;
+
+	FTestPipeline& LoginPipeline = GetLoginPipeline(AccountId);
+	QueryOffersHelperParams.OpParams->LocalAccountId = AccountId;
+	FString CatalogId;
+	GConfig->GetString(TEXT("EOSSDK.Platform.OnlineTests"), TEXT("OfferNamespace0Items"), CatalogId, GEngineIni);
+	QueryOffersHelperParams.OpParams->OverrideCatalogNamespace = CatalogId;
+	GetOffersHelperParams.OpParams->LocalAccountId = AccountId;
+
+	bool bLogout = true;
+
+	LoginPipeline
+		.EmplaceStep<FCommerceQueryOffersHelper>(MoveTemp(QueryOffersHelperParams))
+		.EmplaceStep<FCommerceGetOffersHelper>(MoveTemp(GetOffersHelperParams), ExpectedOffersNum);
+
+	RunToCompletion(bLogout);
+}
+
+COMMERCE_TEST_CASE("Verify that GetOffers returns a correct list of one offer if there is only one cached offer", COMMERCE_GETOFFERS_TAG)
+{
+	FAccountId AccountId;
+	TOptional<uint32_t> ExpectedOffersNum = 1;
+
+	FCommerceQueryOffers::Params OpQueryOffersParams;
+	FCommerceQueryOffersHelper::FHelperParams QueryOffersHelperParams;
+	QueryOffersHelperParams.OpParams = &OpQueryOffersParams;
+
+	FCommerceGetOffers::Params OpGetOffersParams;
+	FCommerceGetOffersHelper::FHelperParams GetOffersHelperParams;
+	GetOffersHelperParams.OpParams = &OpGetOffersParams;
+
+	FTestPipeline& LoginPipeline = GetLoginPipeline(AccountId);
+	QueryOffersHelperParams.OpParams->LocalAccountId = AccountId;
+	FString CatalogId;
+	GConfig->GetString(TEXT("EOSSDK.Platform.OnlineTests"), TEXT("OfferNamespace1Items"), CatalogId, GEngineIni);
+	QueryOffersHelperParams.OpParams->OverrideCatalogNamespace = CatalogId;
+	GetOffersHelperParams.OpParams->LocalAccountId = AccountId;
+
+	bool bLogout = true;
+
+	LoginPipeline
+		.EmplaceStep<FCommerceQueryOffersHelper>(MoveTemp(QueryOffersHelperParams))
+		.EmplaceStep<FCommerceGetOffersHelper>(MoveTemp(GetOffersHelperParams), ExpectedOffersNum);
+
+	RunToCompletion(bLogout);
+}
+
+COMMERCE_TEST_CASE("Verify that GetOffers returns the correct list if there are cached offers", COMMERCE_GETOFFERS_TAG)
+{
+	FAccountId AccountId;
+	TOptional<uint32_t> ExpectedOffersNum = 4;
+
+	FCommerceQueryOffers::Params OpQueryOffersParams;
+	FCommerceQueryOffersHelper::FHelperParams QueryOffersHelperParams;
+	QueryOffersHelperParams.OpParams = &OpQueryOffersParams;
+
+	FCommerceGetOffers::Params OpGetOffersParams;
+	FCommerceGetOffersHelper::FHelperParams GetOffersHelperParams;
+	GetOffersHelperParams.OpParams = &OpGetOffersParams;
+
+	FTestPipeline& LoginPipeline = GetLoginPipeline(AccountId);
+	QueryOffersHelperParams.OpParams->LocalAccountId = AccountId;
+	GetOffersHelperParams.OpParams->LocalAccountId = AccountId;
+
+	bool bLogout = true;
+
+	LoginPipeline
+		.EmplaceStep<FCommerceQueryOffersHelper>(MoveTemp(QueryOffersHelperParams))
+		.EmplaceStep<FCommerceGetOffersHelper>(MoveTemp(GetOffersHelperParams), ExpectedOffersNum);
+
+	RunToCompletion(bLogout);
+}
+
+// GetOffersById
+COMMERCE_TEST_CASE("Verify that GetOffersById returns a fail message if the local user is not logged in", COMMERCE_GETOFFERSBYID_TAG)
+{
+	FAccountId AccountId;
+
+	FCommerceGetOffersById::Params OpGetOffersByIdParams;
+	FCommerceGetOffersByIdHelper::FHelperParams GetOffersByIdHelperParams;
+	GetOffersByIdHelperParams.OpParams = &OpGetOffersByIdParams;
+	GetOffersByIdHelperParams.ExpectedError = TOnlineResult<FCommerceGetOffersById>(Errors::NotLoggedIn());
+
+	FTestPipeline& LoginPipeline = GetLoginPipeline(AccountId);
+	GetOffersByIdHelperParams.OpParams->LocalAccountId = AccountId;
+
+	bool bLogout = false;
+
+	LoginPipeline
+		.EmplaceStep<FAuthLogoutStep>(FPlatformMisc::GetPlatformUserForUserIndex(0))
+		.EmplaceStep<FCommerceGetOffersByIdHelper>(MoveTemp(GetOffersByIdHelperParams));
+
+	RunToCompletion(bLogout);
+}
+
+COMMERCE_TEST_CASE("Verify that GetOffersById returns a fail message of the given local user ID does not match the actual local user ID", COMMERCE_GETOFFERSBYID_TAG)
+{
+	FAccountId AccountId;
+
+	FCommerceGetOffersById::Params OpGetOffersByIdParams;
+	FCommerceGetOffersByIdHelper::FHelperParams GetOffersByIdHelperParams;
+	GetOffersByIdHelperParams.OpParams = &OpGetOffersByIdParams;
+	GetOffersByIdHelperParams.ExpectedError = TOnlineResult<FCommerceGetOffersById>(Errors::NotLoggedIn());
+
+	FTestPipeline& LoginPipeline = GetLoginPipeline(AccountId);
+	GetOffersByIdHelperParams.OpParams->LocalAccountId = FAccountId();
+
+	bool bLogout = true;
+
+	LoginPipeline
+		.EmplaceStep<FCommerceGetOffersByIdHelper>(MoveTemp(GetOffersByIdHelperParams));
+
+	RunToCompletion(bLogout);
+}
+
+COMMERCE_TEST_CASE("Verify that GetOffersById returns an empty list when given an empty list of IDs and there are no offers cached", COMMERCE_GETOFFERSBYID_TAG)
+{
+	FAccountId AccountId;
+	TOptional<uint32_t> ExpectedOffersNum = 0;
+
+	FCommerceQueryOffers::Params OpQueryOffersParams;
+	FCommerceQueryOffersHelper::FHelperParams QueryOffersHelperParams;
+	QueryOffersHelperParams.OpParams = &OpQueryOffersParams;
+
+	FCommerceGetOffersById::Params OpGetOffersByIdParams;
+	FCommerceGetOffersByIdHelper::FHelperParams GetOffersByIdHelperParams;
+	GetOffersByIdHelperParams.OpParams = &OpGetOffersByIdParams;
+
+	FTestPipeline& LoginPipeline = GetLoginPipeline(AccountId);
+	QueryOffersHelperParams.OpParams->LocalAccountId = AccountId;
+	FString CatalogId;
+	GConfig->GetString(TEXT("EOSSDK.Platform.OnlineTests"), TEXT("OfferNamespace0Items"), CatalogId, GEngineIni);
+	QueryOffersHelperParams.OpParams->OverrideCatalogNamespace = CatalogId;
+	GetOffersByIdHelperParams.OpParams->LocalAccountId = AccountId;
+	GetOffersByIdHelperParams.OpParams->OfferIds = {};
+
+	bool bLogout = true;
+
+	LoginPipeline
+		.EmplaceStep<FCommerceQueryOffersHelper>(MoveTemp(QueryOffersHelperParams))
+		.EmplaceStep<FCommerceGetOffersByIdHelper>(MoveTemp(GetOffersByIdHelperParams), ExpectedOffersNum);
+
+}
+
+COMMERCE_TEST_CASE("Verify that GetOffersById returns an empty list when given an empty list of IDs and there are multiple cached offers", COMMERCE_GETOFFERSBYID_TAG)
+{
+	FAccountId AccountId;
+	TOptional<uint32_t> ExpectedOffersNum = 0;
+
+	FCommerceQueryOffers::Params OpQueryOffersParams;
+	FCommerceQueryOffersHelper::FHelperParams QueryOffersHelperParams;
+	QueryOffersHelperParams.OpParams = &OpQueryOffersParams;
+
+	FCommerceGetOffersById::Params OpGetOffersByIdParams;
+	FCommerceGetOffersByIdHelper::FHelperParams GetOffersByIdHelperParams;
+	GetOffersByIdHelperParams.OpParams = &OpGetOffersByIdParams;
+
+	FTestPipeline& LoginPipeline = GetLoginPipeline(AccountId);
+	QueryOffersHelperParams.OpParams->LocalAccountId = AccountId;
+	GetOffersByIdHelperParams.OpParams->LocalAccountId = AccountId;
+	GetOffersByIdHelperParams.OpParams->OfferIds = {};
+
+	bool bLogout = true;
+
+	LoginPipeline
+		.EmplaceStep<FCommerceQueryOffersHelper>(MoveTemp(QueryOffersHelperParams))
+		.EmplaceStep<FCommerceGetOffersByIdHelper>(MoveTemp(GetOffersByIdHelperParams), ExpectedOffersNum);
+
+	RunToCompletion(bLogout);
+}
+
+COMMERCE_TEST_CASE("Verify that GetOffersById returns a fail message when given a populated list of IDs and there are no cached offers", COMMERCE_GETOFFERSBYID_TAG)
+{
+	FAccountId AccountId;
+	TOptional<uint32_t> ExpectedOffersNum = 0;
+
+	FCommerceQueryOffers::Params OpQueryOffersParams;
+	FCommerceQueryOffersHelper::FHelperParams QueryOffersHelperParams;
+	QueryOffersHelperParams.OpParams = &OpQueryOffersParams;
+
+	FCommerceGetOffersById::Params OpGetOffersByIdParams;
+	FCommerceGetOffersByIdHelper::FHelperParams GetOffersByIdHelperParams;
+	GetOffersByIdHelperParams.OpParams = &OpGetOffersByIdParams;
+	GetOffersByIdHelperParams.ExpectedError = TOnlineResult<FCommerceGetOffersById>(Errors::NotFound());
+
+	FTestPipeline& LoginPipeline = GetLoginPipeline(AccountId);
+	QueryOffersHelperParams.OpParams->LocalAccountId = AccountId;
+	FString CatalogId;
+	GConfig->GetString(TEXT("EOSSDK.Platform.OnlineTests"), TEXT("OfferNamespace0Items"), CatalogId, GEngineIni);
+	QueryOffersHelperParams.OpParams->OverrideCatalogNamespace = CatalogId;
+	GetOffersByIdHelperParams.OpParams->LocalAccountId = AccountId;
+	FOfferId OfferId1NotExisting;
+	FOfferId OfferId2NotExisting;
+	GConfig->GetString(TEXT("EOSSDK.Platform.OnlineTests"), TEXT("OfferId1NotExisting"), OfferId1NotExisting, GEngineIni);
+	GConfig->GetString(TEXT("EOSSDK.Platform.OnlineTests"), TEXT("OfferId2NotExisting"), OfferId2NotExisting, GEngineIni);
+	GetOffersByIdHelperParams.OpParams->OfferIds = { OfferId1NotExisting, OfferId2NotExisting };
+
+	bool bLogout = true;
+
+	LoginPipeline
+		.EmplaceStep<FCommerceQueryOffersHelper>(MoveTemp(QueryOffersHelperParams))
+		.EmplaceStep<FCommerceGetOffersByIdHelper>(MoveTemp(GetOffersByIdHelperParams), ExpectedOffersNum);
+
+	RunToCompletion(bLogout);
+}
+
+COMMERCE_TEST_CASE("Verify that GetOffersById returns an empty list when given a populated list of IDs and none of them match any cached offers", COMMERCE_GETOFFERSBYID_TAG)
+{
+	FAccountId AccountId;
+	TOptional<uint32_t> ExpectedOffersNum = 0;
+
+	FCommerceQueryOffers::Params OpQueryOffersParams;
+	FCommerceQueryOffersHelper::FHelperParams QueryOffersHelperParams;
+	QueryOffersHelperParams.OpParams = &OpQueryOffersParams;
+
+	FCommerceGetOffersById::Params OpGetOffersByIdParams;
+	FCommerceGetOffersByIdHelper::FHelperParams GetOffersByIdHelperParams;
+	GetOffersByIdHelperParams.OpParams = &OpGetOffersByIdParams;
+
+	FTestPipeline& LoginPipeline = GetLoginPipeline(AccountId);
+	QueryOffersHelperParams.OpParams->LocalAccountId = AccountId;
+	GetOffersByIdHelperParams.OpParams->LocalAccountId = AccountId;
+	FOfferId OfferId1NotExisting;
+	FOfferId OfferId2NotExisting;
+	GConfig->GetString(TEXT("EOSSDK.Platform.OnlineTests"), TEXT("OfferId1NotExisting"), OfferId1NotExisting, GEngineIni);
+	GConfig->GetString(TEXT("EOSSDK.Platform.OnlineTests"), TEXT("OfferId2NotExisting"), OfferId2NotExisting, GEngineIni);
+	GetOffersByIdHelperParams.OpParams->OfferIds = { OfferId1NotExisting, OfferId2NotExisting };
+
+	bool bLogout = true;
+
+	LoginPipeline
+		.EmplaceStep<FCommerceQueryOffersHelper>(MoveTemp(QueryOffersHelperParams))
+		.EmplaceStep<FCommerceGetOffersByIdHelper>(MoveTemp(GetOffersByIdHelperParams), ExpectedOffersNum);
+
+	RunToCompletion(bLogout);
+}
+
+COMMERCE_TEST_CASE("Verify that GetOffersById returns correct list of one offer when given a populated list of IDs where one ID is cached and another is not", COMMERCE_GETOFFERSBYID_TAG)
+{
+	FAccountId AccountId;
+	TOptional<uint32_t> ExpectedOffersNum = 1;
+
+	FCommerceQueryOffers::Params OpQueryOffersParams;
+	FCommerceQueryOffersHelper::FHelperParams QueryOffersHelperParams;
+	QueryOffersHelperParams.OpParams = &OpQueryOffersParams;
+
+	FCommerceGetOffersById::Params OpGetOffersByIdParams;
+	FCommerceGetOffersByIdHelper::FHelperParams GetOffersByIdHelperParams;
+	GetOffersByIdHelperParams.OpParams = &OpGetOffersByIdParams;
+
+	FTestPipeline& LoginPipeline = GetLoginPipeline(AccountId);
+	QueryOffersHelperParams.OpParams->LocalAccountId = AccountId;
+	GetOffersByIdHelperParams.OpParams->LocalAccountId = AccountId;
+	FOfferId OfferId1NotExisting;
+	FOfferId OfferId2;
+	GConfig->GetString(TEXT("EOSSDK.Platform.OnlineTests"), TEXT("OfferId1NotExisting"), OfferId1NotExisting, GEngineIni);
+	GConfig->GetString(TEXT("EOSSDK.Platform.OnlineTests"), TEXT("OfferId2"), OfferId2, GEngineIni);
+	GetOffersByIdHelperParams.OpParams->OfferIds = { OfferId1NotExisting, OfferId2 };
+
+	bool bLogout = true;
+
+	LoginPipeline
+		.EmplaceStep<FCommerceQueryOffersHelper>(MoveTemp(QueryOffersHelperParams))
+		.EmplaceStep<FCommerceGetOffersByIdHelper>(MoveTemp(GetOffersByIdHelperParams), ExpectedOffersNum);
+
+	RunToCompletion(bLogout);
+}
+
+COMMERCE_TEST_CASE("Verify that GetOffersById returns the correct list of one offer when given the ID for one cached offer", COMMERCE_GETOFFERSBYID_TAG)
+{
+	FAccountId AccountId;
+	TOptional<uint32_t> ExpectedOffersNum = 1;
+
+	FCommerceQueryOffers::Params OpQueryOffersParams;
+	FCommerceQueryOffersHelper::FHelperParams QueryOffersHelperParams;
+	QueryOffersHelperParams.OpParams = &OpQueryOffersParams;
+
+	FCommerceGetOffersById::Params OpGetOffersByIdParams;
+	FCommerceGetOffersByIdHelper::FHelperParams GetOffersByIdHelperParams;
+	GetOffersByIdHelperParams.OpParams = &OpGetOffersByIdParams;
+
+	FTestPipeline& LoginPipeline = GetLoginPipeline(AccountId);
+	QueryOffersHelperParams.OpParams->LocalAccountId = AccountId;
+	GetOffersByIdHelperParams.OpParams->LocalAccountId = AccountId;
+	FOfferId OfferId1;
+	GConfig->GetString(TEXT("EOSSDK.Platform.OnlineTests"), TEXT("OfferId1"), OfferId1, GEngineIni);
+	GetOffersByIdHelperParams.OpParams->OfferIds = {OfferId1};
+
+	bool bLogout = true;
+
+	LoginPipeline
+		.EmplaceStep<FCommerceQueryOffersHelper>(MoveTemp(QueryOffersHelperParams))
+		.EmplaceStep<FCommerceGetOffersByIdHelper>(MoveTemp(GetOffersByIdHelperParams), ExpectedOffersNum);
+
+	RunToCompletion(bLogout);
+}
+
+COMMERCE_TEST_CASE("Verify that GetOffersById returns the correct list of offers when given a populated list of multiple cached IDs", COMMERCE_GETOFFERSBYID_TAG)
+{
+	FAccountId AccountId;
+	TOptional<uint32_t> ExpectedOffersNum = 2;
+
+	FCommerceQueryOffers::Params OpQueryOffersParams;
+	FCommerceQueryOffersHelper::FHelperParams QueryOffersHelperParams;
+	QueryOffersHelperParams.OpParams = &OpQueryOffersParams;
+
+	FCommerceGetOffersById::Params OpGetOffersByIdParams;
+	FCommerceGetOffersByIdHelper::FHelperParams GetOffersByIdHelperParams;
+	GetOffersByIdHelperParams.OpParams = &OpGetOffersByIdParams;
+
+	FTestPipeline& LoginPipeline = GetLoginPipeline(AccountId);
+	QueryOffersHelperParams.OpParams->LocalAccountId = AccountId;
+	GetOffersByIdHelperParams.OpParams->LocalAccountId = AccountId;
+	FOfferId OfferId1;
+	FOfferId OfferId2;
+	GConfig->GetString(TEXT("EOSSDK.Platform.OnlineTests"), TEXT("OfferId1"), OfferId1, GEngineIni);
+	GConfig->GetString(TEXT("EOSSDK.Platform.OnlineTests"), TEXT("OfferId2"), OfferId2, GEngineIni);
+	GetOffersByIdHelperParams.OpParams->OfferIds = { OfferId1, OfferId2 };
+
+	bool bLogout = true;
+
+	LoginPipeline
+		.EmplaceStep<FCommerceQueryOffersHelper>(MoveTemp(QueryOffersHelperParams))
+		.EmplaceStep<FCommerceGetOffersByIdHelper>(MoveTemp(GetOffersByIdHelperParams), ExpectedOffersNum);
+
+	RunToCompletion(bLogout);
+}
+
+// ShowStoreUI tests
+COMMERCE_TEST_CASE("Verify that ShowStoreUI returns a fail message if the local user is not logged in", COMMERCE_DISABLED_TAG)
+{
+	// TODO
+	// The ShowStoreUI EOS interface is not implemented
+}
+
+COMMERCE_TEST_CASE("Verify that ShowStoreUI returns a fail message of the given local user ID does not match the actual local user ID", COMMERCE_DISABLED_TAG)
+{
+	// TODO
+	// The ShowStoreUI EOS interface is not implemented
+}
+
+COMMERCE_TEST_CASE("Verify that ShowStoreUI displays platform store UI", COMMERCE_DISABLED_TAG) // May not be possible to automate currently
+{
+	// TODO
+	// The ShowStoreUI EOS interface is not implemented
+}
+
+// Checkout tests
+COMMERCE_TEST_CASE("Verify that Checkout returns a fail message if the local user is not logged in")
+{
+	FAccountId AccountId;
+
+	FCommerceCheckout::Params OpCheckoutParams;
+	FCommerceCheckoutHelper::FHelperParams CheckoutHelperParams;
+	CheckoutHelperParams.OpParams = &OpCheckoutParams;
+	CheckoutHelperParams.ExpectedError = TOnlineResult<FCommerceCheckout>(Errors::NotLoggedIn());
+
+	FTestPipeline& LoginPipeline = GetLoginPipeline(AccountId);
+	CheckoutHelperParams.OpParams->LocalAccountId = AccountId;
+
+	bool bLogout = false;
+
+	LoginPipeline
+		.EmplaceStep<FAuthLogoutStep>(FPlatformMisc::GetPlatformUserForUserIndex(0))
+		.EmplaceStep<FCommerceCheckoutHelper>(MoveTemp(CheckoutHelperParams));
+
+	RunToCompletion(bLogout);
+}
+
+COMMERCE_TEST_CASE("Verify that Checkout returns a fail message of the given local user ID does not match the actual local user ID")
+{
+	FAccountId AccountId;
+
+	FCommerceCheckout::Params OpCheckoutParams;
+	FCommerceCheckoutHelper::FHelperParams CheckoutHelperParams;
+	CheckoutHelperParams.OpParams = &OpCheckoutParams;
+	CheckoutHelperParams.ExpectedError = TOnlineResult<FCommerceCheckout>(Errors::NotLoggedIn());
+
+	FTestPipeline& LoginPipeline = GetLoginPipeline(AccountId);
+	CheckoutHelperParams.OpParams->LocalAccountId = FAccountId();
+
+	bool bLogout = true;
+
+	LoginPipeline
+		.EmplaceStep<FCommerceCheckoutHelper>(MoveTemp(CheckoutHelperParams));
+
+	RunToCompletion(bLogout);
+}
+
+COMMERCE_TEST_CASE("Verify that Checkout initiates the checkout process when given an Offers array with one offer")
+{
+	/*
+	FAccountId AccountId;
+	TOptional<uint32_t> ExpectedOffersNum = 1;
+
+	FCommerceQueryOffers::Params OpQueryOffersParams;
+	FCommerceQueryOffersHelper::FHelperParams QueryOffersHelperParams;
+	QueryOffersHelperParams.OpParams = &OpQueryOffersParams;
+
+	FCommerceGetOffersById::Params OpGetOffersByIdParams;
+	FCommerceGetOffersByIdHelper::FHelperParams GetOffersByIdHelperParams;
+	GetOffersByIdHelperParams.OpParams = &OpGetOffersByIdParams;
+
+	FCommerceCheckout::Params OpCheckoutParams;
+	FCommerceCheckoutHelper::FHelperParams CheckoutHelperParams;
+	CheckoutHelperParams.OpParams = &OpCheckoutParams;
+
+	FTestPipeline& LoginPipeline = GetLoginPipeline(AccountId);
+
+	QueryOffersHelperParams.OpParams->LocalAccountId = AccountId;
+
+	GetOffersByIdHelperParams.OpParams->LocalAccountId = AccountId;
+	FOfferId OfferId2;
+	GConfig->GetString(TEXT("EOSSDK.Platform.OnlineTests"), TEXT("OfferId2"), OfferId2, GEngineIni);
+	GetOffersByIdHelperParams.OpParams->OfferIds = { OfferId2 };
+
+	FPurchaseOffer Purchase = { OfferId2, 1};
+	CheckoutHelperParams.OpParams->LocalAccountId = AccountId;
+	CheckoutHelperParams.OpParams->Offers = { Purchase };
+
+	bool bLogout = true;
+
+	LoginPipeline
+		.EmplaceStep<FCommerceQueryOffersHelper>(MoveTemp(QueryOffersHelperParams))
+		.EmplaceStep<FCommerceGetOffersByIdHelper>(MoveTemp(GetOffersByIdHelperParams), ExpectedOffersNum)
+		.EmplaceStep<FCommerceCheckoutHelper>(MoveTemp(CheckoutHelperParams));
+
+	RunToCompletion(bLogout);
+	*/
+}
+
+COMMERCE_TEST_CASE("Verify that Checkout initiates the checkout process when given an Offers array with multiple offers")
+{
+	// TODO
+}
+
+COMMERCE_TEST_CASE("Verify that Checkout does not initiate the checkout process when given an empty Offers array")
+{
+	// TODO
+}
+
+COMMERCE_TEST_CASE("Verify that Checkout returns the correct TransactionId after a completed purchase")
+{
+	// TODO
+}
+
+COMMERCE_TEST_CASE("Verify that Checkout does not return a TransactionId if the purchase is incomplete")
+{
+	// TODO
+}
+
+// QueryTransactionEntitlements tests
+COMMERCE_TEST_CASE("Verify that QueryTransactionEntitlements returns a fail message if the local user is not logged in")
+{
+	// TODO
+}
+
+COMMERCE_TEST_CASE("Verify that QueryTransactionEntitlements returns a fail message of the given local user ID does not match the actual local user ID")
+{
+	// TODO
+}
+
+COMMERCE_TEST_CASE("Verify that QueryTransactionEntitlements returns the correct list of entitlements when given an existing TransactionId and caches them")
+{
+	// TODO
+}
+
+COMMERCE_TEST_CASE("Verify that QueryTransactionEntitlements returns a fail message when given a nonexisting TransactionId")
+{
+	// TODO
+}
+
+// QueryEntitlements tests
+COMMERCE_TEST_CASE("Verify that QueryEntitlements returns a fail message if the local user is not logged in")
+{
+	FAccountId AccountId;
+
+	FCommerceQueryEntitlements::Params OpQueryEntitlementsParams;
+	FCommerceQueryEntitlementsHelper::FHelperParams QueryEntitlementsHelperParams;
+	QueryEntitlementsHelperParams.ExpectedError = TOnlineResult<FCommerceQueryEntitlements>(Errors::NotLoggedIn());
+	QueryEntitlementsHelperParams.OpParams = &OpQueryEntitlementsParams;
+
+	FTestPipeline& LoginPipeline = GetLoginPipeline(AccountId);
+	QueryEntitlementsHelperParams.OpParams->LocalAccountId = AccountId;
+
+	bool bLogout = false;
+
+	LoginPipeline
+		.EmplaceStep<FAuthLogoutStep>(FPlatformMisc::GetPlatformUserForUserIndex(0))
+		.EmplaceStep<FCommerceQueryEntitlementsHelper>(MoveTemp(QueryEntitlementsHelperParams));
+
+	RunToCompletion(bLogout);
+}
+
+COMMERCE_TEST_CASE("Verify that QueryEntitlements returns a fail message of the given local user ID does not match the actual local user ID")
+{
+	FAccountId AccountId;
+
+	FCommerceQueryEntitlements::Params OpQueryEntitlementsParams;
+	FCommerceQueryEntitlementsHelper::FHelperParams QueryEntitlementsHelperParams;
+	QueryEntitlementsHelperParams.OpParams = &OpQueryEntitlementsParams;
+	QueryEntitlementsHelperParams.ExpectedError = TOnlineResult<FCommerceQueryEntitlements>(Errors::NotLoggedIn());
+
+	FTestPipeline& LoginPipeline = GetLoginPipeline(AccountId);
+	QueryEntitlementsHelperParams.OpParams->LocalAccountId = FAccountId();
+
+	bool bLogout = true;
+
+	LoginPipeline
+		.EmplaceStep<FCommerceQueryEntitlementsHelper>(MoveTemp(QueryEntitlementsHelperParams));
+
+	RunToCompletion(bLogout);
+}
+
+COMMERCE_TEST_CASE("Verify that QueryEntitlements caches all entitlements")
+{
+	FAccountId AccountId;
+	TOptional<uint32_t> ExpectedEntitlementsNum = 0;
+
+	FCommerceQueryEntitlements::Params OpQueryEntitlementsParams;
+	FCommerceQueryEntitlementsHelper::FHelperParams QueryEntitlementsHelperParams;
+	QueryEntitlementsHelperParams.OpParams = &OpQueryEntitlementsParams;
+
+	FCommerceGetEntitlements::Params OpGetEntitlementsParams;
+	FCommerceGetEntitlementsHelper::FHelperParams GetEntitlementsHelperParams;
+	GetEntitlementsHelperParams.OpParams = &OpGetEntitlementsParams;
+
+	FTestPipeline& LoginPipeline = GetLoginPipeline(AccountId);
+	QueryEntitlementsHelperParams.OpParams->LocalAccountId = AccountId;
+	QueryEntitlementsHelperParams.OpParams->bIncludeRedeemed = true;
+	GetEntitlementsHelperParams.OpParams->LocalAccountId = AccountId;
+
+	bool bLogout = true;
+
+	LoginPipeline
+		.EmplaceStep<FCommerceQueryEntitlementsHelper>(MoveTemp(QueryEntitlementsHelperParams))
+		.EmplaceStep<FCommerceGetEntitlementsHelper>(MoveTemp(GetEntitlementsHelperParams), ExpectedEntitlementsNum);
+
+	RunToCompletion(bLogout);
+}
+
+// GetEntitlements tests
+COMMERCE_TEST_CASE("Verify that GetEntitlements returns a fail message if the local user is not logged in")
+{
+	FAccountId AccountId;
+
+	FCommerceGetEntitlements::Params OpGetEntitlementsParams;
+	FCommerceGetEntitlementsHelper::FHelperParams GetEntitlementsHelperParams;
+	GetEntitlementsHelperParams.OpParams = &OpGetEntitlementsParams;
+	GetEntitlementsHelperParams.ExpectedError = TOnlineResult<FCommerceGetEntitlements>(Errors::NotLoggedIn());
+
+	FTestPipeline& LoginPipeline = GetLoginPipeline(AccountId);
+	GetEntitlementsHelperParams.OpParams->LocalAccountId = AccountId;
+
+	bool bLogout = false;
+
+	LoginPipeline
+		.EmplaceStep<FAuthLogoutStep>(FPlatformMisc::GetPlatformUserForUserIndex(0))
+		.EmplaceStep<FCommerceGetEntitlementsHelper>(MoveTemp(GetEntitlementsHelperParams));
+
+	RunToCompletion(bLogout);
+}
+
+COMMERCE_TEST_CASE("Verify that GetEntitlements returns a fail message of the given local user ID does not match the actual local user ID")
+{
+	FAccountId AccountId;
+
+	FCommerceGetEntitlements::Params OpGetEntitlementsParams;
+	FCommerceGetEntitlementsHelper::FHelperParams GetEntitlementsHelperParams;
+	GetEntitlementsHelperParams.OpParams = &OpGetEntitlementsParams;
+	GetEntitlementsHelperParams.ExpectedError = TOnlineResult<FCommerceGetEntitlements>(Errors::NotLoggedIn());
+
+	FTestPipeline& LoginPipeline = GetLoginPipeline(AccountId);
+	GetEntitlementsHelperParams.OpParams->LocalAccountId = FAccountId();
+
+	bool bLogout = true;
+
+	LoginPipeline
+		.EmplaceStep<FCommerceGetEntitlementsHelper>(MoveTemp(GetEntitlementsHelperParams));
+
+	RunToCompletion(bLogout);
+}
+
+
+COMMERCE_TEST_CASE("Verify that GetEntitlements returns a list of all cached entitlements")
+{
+	// TODO
+}
+
+COMMERCE_TEST_CASE("Verify that GetEntitlements returns an empty list when no entitlements are cached")
+{
+	// TODO
+}
+
+// RedeemEntitlement tests
+COMMERCE_TEST_CASE("Verify that RedeemEntitlement returns a fail message if the local user is not logged in")
+{
+	// TODO
+}
+
+COMMERCE_TEST_CASE("Verify that RedeemEntitlement returns a fail message of the given local user ID does not match the actual local user ID")
+{
+	// TODO
+}
+
+COMMERCE_TEST_CASE("Verify that RedeemEntitlement redeems no entitlements when given an existing EntitlementId and Quantity 0")
+{
+	// TODO
+}
+
+COMMERCE_TEST_CASE("Verify that RedeemEntitlement redeems 1 of the correct entitlements when given an existing EntitlementId and Quantity 1")
+{
+	// TODO
+}
+
+COMMERCE_TEST_CASE("Verify that RedeemEntitlement redeems 2 of the correct entitlements when given an existing EntitlementId and Quantity 2")
+{
+	// TODO
+}
+
+COMMERCE_TEST_CASE("Verify that RedeemEntitlement redeems all of the given entitlements when given a Quantity equal to the available number for that entitlement")
+{
+	// TODO
+}
+
+COMMERCE_TEST_CASE("Verify that RedeemEntitlement returns a fail message when given a Quantity that is 1 more than the available number for the given entitlemenet")
+{
+	// TODO
+}
+
+COMMERCE_TEST_CASE("Verify that RedeemEntitlement returns a fail message when given a Quantity that is 2 more than the available number for the given entitlemenet")
+{
+	// TODO
+}
+
+COMMERCE_TEST_CASE("Verify that RedeemEntitlement returns a fail message when given a nonexisting EntitlementId")
+{
+	// TODO
+}
+
+COMMERCE_TEST_CASE("Verify that RedeemEntitlement returns all of the given entitlements when given a Quantity matching the number of entitlements available")
+{
+	// TODO
+}
+
+// RetrieveS2SToken tests
+COMMERCE_TEST_CASE("Verify that RetrieveS2SToken returns a fail message if the local user is not logged in")
+{
+	// TODO
+}
+
+COMMERCE_TEST_CASE("Verify that RetrieveS2SToken returns a fail message of the given local user ID does not match the actual local user ID")
+{
+	// TODO
+}
+
+COMMERCE_TEST_CASE("Verify that RetrieveS2SToken successfully returns the correct Token for the given TokenType")
+{
+	// TODO
+}
+
+COMMERCE_TEST_CASE("Verify that RetrieveS2SToken returns a fail message if the given TokenType does not exist")
+{
+	// TODO
+}
+
+// OnPurchaseCompleted tests
+COMMERCE_TEST_CASE("Verify that OnPurchaseCompleted returns a fail message if the local user is not logged in")
+{
+	// TODO
+}
+
+COMMERCE_TEST_CASE("Verify that OnPurchaseCompleted returns a fail message of the given local user ID does not match the actual local user ID")
+{
+	// TODO
+}
+
+COMMERCE_TEST_CASE("Verify that OnPurchaseCompleted completes without an error")
+{
+	// TODO
+}
diff --git a/Engine/Source/Programs/Online/OnlineTestsCore/Source/Public/Helpers/Commerce/CommerceCheckoutHelper.h b/Engine/Source/Programs/Online/OnlineTestsCore/Source/Public/Helpers/Commerce/CommerceCheckoutHelper.h
new file mode 100644
--- /dev/null
+++ b/Engine/Source/Programs/Online/OnlineTestsCore/Source/Public/Helpers/Commerce/CommerceCheckoutHelper.h
@@ -0,0 +1,53 @@
+// Copyright Epic Games, Inc. All Rights Reserved.
+
+#pragma once
+
+#include "Online/CommerceCommon.h"
+#include "AsyncTestStep.h"
+#include "OnlineCatchHelper.h"
+
+struct FCommerceCheckoutHelper : public FAsyncTestStep
+{
+	using ParamsType = UE::Online::FCommerceCheckout::Params;
+	using ResultType = UE::Online::TOnlineResult<FCommerceCheckout>;
+
+	struct FHelperParams
+	{
+		ParamsType* OpParams = nullptr;
+		TOptional<ResultType> ExpectedError;
+	};
+
+	FCommerceCheckoutHelper(FHelperParams&& InHelperParams)
+		: HelperParams(MoveTemp(InHelperParams))
+	{
+		REQUIRE(HelperParams.OpParams);
+		REQUIRE((!HelperParams.ExpectedError.IsSet() || HelperParams.ExpectedError->IsError()));
+	}
+
+	virtual ~FCommerceCheckoutHelper() = default;
+
+	virtual void Run(FAsyncStepResult Promise, SubsystemType Services) override
+	{
+		CommerceInterface = Services->GetCommerceInterface();
+		REQUIRE(CommerceInterface);
+
+		UE::Online::TOnlineAsyncOpHandle<UE::Online::FCommerceCheckout> CheckoutResult = CommerceInterface->Checkout(MoveTemp(*HelperParams.OpParams))
+			.OnComplete([this, Promise = MoveTemp(Promise)](const ResultType& Result)
+				{
+					if (HelperParams.ExpectedError.IsSet())
+					{
+						REQUIRE_OP_EQ(Result, HelperParams.ExpectedError->GetErrorValue());
+					}
+					else
+					{
+						CHECK(Result.GetOkValue().TransactionId.IsSet());
+					}
+					Promise->SetValue(true);
+				});
+	}
+
+protected:
+	FHelperParams HelperParams;
+	UE::Online::ICommercePtr CommerceInterface = nullptr;
+	TOptional<FString> TransactionId;
+};
\ No newline at end of file
diff --git a/Engine/Source/Programs/Online/OnlineTestsCore/Source/Public/Helpers/Commerce/CommerceGetEntitlementsHelper.h b/Engine/Source/Programs/Online/OnlineTestsCore/Source/Public/Helpers/Commerce/CommerceGetEntitlementsHelper.h
new file mode 100644
--- /dev/null
+++ b/Engine/Source/Programs/Online/OnlineTestsCore/Source/Public/Helpers/Commerce/CommerceGetEntitlementsHelper.h
@@ -0,0 +1,57 @@
+// Copyright Epic Games, Inc. All Rights Reserved.
+
+#pragma once
+
+#include "Online/CommerceCommon.h"
+#include "AsyncTestStep.h"
+#include "OnlineCatchHelper.h"
+
+struct FCommerceGetEntitlementsHelper : public FAsyncTestStep
+{
+	using ParamsType = UE::Online::FCommerceGetEntitlements::Params;
+	using ResultType = UE::Online::TOnlineResult<UE::Online::FCommerceGetEntitlements>;
+
+	struct FHelperParams
+	{
+		ParamsType* OpParams = nullptr;
+		TOptional<ResultType> ExpectedError;
+	};
+
+	FCommerceGetEntitlementsHelper(FHelperParams&& InHelperParams, const TOptional<uint32_t> InExpectedEntitlementsNum = TOptional<uint32_t>())
+		: HelperParams(MoveTemp(InHelperParams))
+		, ExpectedEntitlementsNum(InExpectedEntitlementsNum)
+	{
+		REQUIRE(HelperParams.OpParams);
+		REQUIRE((!HelperParams.ExpectedError.IsSet() || HelperParams.ExpectedError->IsError()));
+	}
+
+	virtual ~FCommerceGetEntitlementsHelper() = default;
+
+	virtual void Run(FAsyncStepResult Promise, SubsystemType Services) override
+	{
+		CommerceInterface = Services->GetCommerceInterface();
+		REQUIRE(CommerceInterface);
+
+		ResultType Result = CommerceInterface->GetEntitlements(MoveTemp(*HelperParams.OpParams));
+
+		if (HelperParams.ExpectedError.IsSet())
+		{
+			REQUIRE_OP_EQ(Result, HelperParams.ExpectedError->GetErrorValue());
+		}
+		else if (ExpectedEntitlementsNum.IsSet())
+		{
+			CHECK(Result.GetOkValue().Entitlements.Num() == ExpectedEntitlementsNum.GetValue());
+		}
+		else
+		{
+			CHECK(Result.GetOkValue().Entitlements.IsEmpty());
+		}
+		Promise->SetValue(true);
+	}
+
+protected:
+	FHelperParams HelperParams;
+	TOptional<uint32_t> ExpectedEntitlementsNum = TOptional<uint32_t>();
+	UE::Online::ICommercePtr CommerceInterface = nullptr;
+	TArray<FEntitlement> Entitlements;
+};
\ No newline at end of file
diff --git a/Engine/Source/Programs/Online/OnlineTestsCore/Source/Public/Helpers/Commerce/CommerceGetOffersByIdHelper.h b/Engine/Source/Programs/Online/OnlineTestsCore/Source/Public/Helpers/Commerce/CommerceGetOffersByIdHelper.h
new file mode 100644
--- /dev/null
+++ b/Engine/Source/Programs/Online/OnlineTestsCore/Source/Public/Helpers/Commerce/CommerceGetOffersByIdHelper.h
@@ -0,0 +1,57 @@
+// Copyright Epic Games, Inc. All Rights Reserved.
+
+#pragma once
+
+#include "Online/CommerceCommon.h"
+#include "AsyncTestStep.h"
+#include "OnlineCatchHelper.h"
+
+struct FCommerceGetOffersByIdHelper : public FAsyncTestStep
+{
+	using ParamsType = UE::Online::FCommerceGetOffersById::Params;
+	using ResultType = UE::Online::TOnlineResult<UE::Online::FCommerceGetOffersById>;
+
+	struct FHelperParams
+	{
+		ParamsType* OpParams = nullptr;
+		TOptional<ResultType> ExpectedError;
+	};
+
+	FCommerceGetOffersByIdHelper(FHelperParams&& InHelperParams, const TOptional<uint32_t> InExpectedOffersNum = TOptional<uint32_t>())
+		: HelperParams(MoveTemp(InHelperParams))
+		, ExpectedOffersNum(InExpectedOffersNum)
+	{
+		REQUIRE(HelperParams.OpParams);
+		REQUIRE((!HelperParams.ExpectedError.IsSet() || HelperParams.ExpectedError->IsError()));
+	}
+
+	virtual ~FCommerceGetOffersByIdHelper() = default;
+
+	virtual void Run(FAsyncStepResult Promise, SubsystemType Services) override
+	{
+		CommerceInterface = Services->GetCommerceInterface();
+		REQUIRE(CommerceInterface);
+
+		ResultType Result = CommerceInterface->GetOffersById(MoveTemp(*HelperParams.OpParams));
+
+		if (HelperParams.ExpectedError.IsSet())
+		{
+			REQUIRE_OP_EQ(Result, HelperParams.ExpectedError->GetErrorValue());
+		}
+		else if (ExpectedOffersNum.IsSet())
+		{
+			CHECK(Result.GetOkValue().Offers.Num() == ExpectedOffersNum.GetValue());
+		}
+		else
+		{
+			CHECK(Result.GetOkValue().Offers.IsEmpty());
+		}
+		Promise->SetValue(true);
+	}
+
+protected:
+	FHelperParams HelperParams;
+	TOptional<uint32_t> ExpectedOffersNum = TOptional<uint32_t>();
+	UE::Online::ICommercePtr CommerceInterface = nullptr;
+	TArray<FOffer> Offers;
+};
\ No newline at end of file
diff --git a/Engine/Source/Programs/Online/OnlineTestsCore/Source/Public/Helpers/Commerce/CommerceGetOffersHelper.h b/Engine/Source/Programs/Online/OnlineTestsCore/Source/Public/Helpers/Commerce/CommerceGetOffersHelper.h
new file mode 100644
--- /dev/null
+++ b/Engine/Source/Programs/Online/OnlineTestsCore/Source/Public/Helpers/Commerce/CommerceGetOffersHelper.h
@@ -0,0 +1,57 @@
+// Copyright Epic Games, Inc. All Rights Reserved.
+
+#pragma once
+
+#include "Online/CommerceCommon.h"
+#include "AsyncTestStep.h"
+#include "OnlineCatchHelper.h"
+
+struct FCommerceGetOffersHelper : public FAsyncTestStep
+{
+	using ParamsType = UE::Online::FCommerceGetOffers::Params;
+	using ResultType = UE::Online::TOnlineResult<UE::Online::FCommerceGetOffers>;
+
+	struct FHelperParams
+	{
+		ParamsType* OpParams = nullptr;
+		TOptional<ResultType> ExpectedError;
+	};
+
+	FCommerceGetOffersHelper(FHelperParams&& InHelperParams, const TOptional<uint32_t> InExpectedOffersNum = TOptional<uint32_t>())
+		: HelperParams(MoveTemp(InHelperParams))
+		, ExpectedOffersNum(InExpectedOffersNum)
+	{
+		REQUIRE(HelperParams.OpParams);
+		REQUIRE((!HelperParams.ExpectedError.IsSet() || HelperParams.ExpectedError->IsError()));
+	}
+
+	virtual ~FCommerceGetOffersHelper() = default;
+
+	virtual void Run(FAsyncStepResult Promise, SubsystemType Services) override
+	{
+		CommerceInterface = Services->GetCommerceInterface();
+		REQUIRE(CommerceInterface);
+
+		ResultType Result = CommerceInterface->GetOffers(MoveTemp(*HelperParams.OpParams));
+
+		if (HelperParams.ExpectedError.IsSet())
+		{
+			REQUIRE_OP_EQ(Result, HelperParams.ExpectedError->GetErrorValue());
+		}
+		else if (ExpectedOffersNum.IsSet())
+		{
+			CHECK(Result.GetOkValue().Offers.Num() == ExpectedOffersNum.GetValue());
+		}
+		else
+		{
+			CHECK(Result.GetOkValue().Offers.IsEmpty());
+		}
+		Promise->SetValue(true);
+	}
+
+protected:
+	FHelperParams HelperParams;
+	TOptional<uint32_t> ExpectedOffersNum = TOptional<uint32_t>();
+	UE::Online::ICommercePtr CommerceInterface = nullptr;
+	TArray<FOffer> Offers;
+};
\ No newline at end of file
diff --git a/Engine/Source/Programs/Online/OnlineTestsCore/Source/Public/Helpers/Commerce/CommerceQueryEntitlementsHelper.h b/Engine/Source/Programs/Online/OnlineTestsCore/Source/Public/Helpers/Commerce/CommerceQueryEntitlementsHelper.h
new file mode 100644
--- /dev/null
+++ b/Engine/Source/Programs/Online/OnlineTestsCore/Source/Public/Helpers/Commerce/CommerceQueryEntitlementsHelper.h
@@ -0,0 +1,52 @@
+// Copyright Epic Games, Inc. All Rights Reserved.
+
+#pragma once
+
+#include "Online/CommerceCommon.h"
+#include "AsyncTestStep.h"
+#include "OnlineCatchHelper.h"
+
+struct FCommerceQueryEntitlementsHelper : public FAsyncTestStep
+{
+	using ParamsType = UE::Online::FCommerceQueryEntitlements::Params;
+	using ResultType = UE::Online::TOnlineResult<UE::Online::FCommerceQueryEntitlements>;
+
+	struct FHelperParams
+	{
+		ParamsType* OpParams = nullptr;
+		TOptional<ResultType> ExpectedError;
+	};
+
+	FCommerceQueryEntitlementsHelper(FHelperParams&& InHelperParams)
+		: HelperParams(MoveTemp(InHelperParams))
+	{
+		REQUIRE(HelperParams.OpParams);
+		REQUIRE((!HelperParams.ExpectedError.IsSet() || HelperParams.ExpectedError->IsError()));
+	}
+
+	virtual ~FCommerceQueryEntitlementsHelper() = default;
+
+	virtual void Run(FAsyncStepResult Promise, SubsystemType Services) override
+	{
+		CommerceInterface = Services->GetCommerceInterface();
+		REQUIRE(CommerceInterface);
+
+		CommerceInterface->QueryEntitlements(MoveTemp(*HelperParams.OpParams))
+			.OnComplete([this, Promise = MoveTemp(Promise)](const ResultType& Result)
+				{
+					if (!HelperParams.ExpectedError.IsSet())
+					{
+						REQUIRE(Result.IsOk());
+					}
+					else
+					{
+						REQUIRE_OP_EQ(Result, HelperParams.ExpectedError->GetErrorValue());
+					}
+					Promise->SetValue(true);
+				});
+	}
+
+protected:
+	FHelperParams HelperParams;
+	UE::Online::ICommercePtr CommerceInterface = nullptr;
+};
\ No newline at end of file
diff --git a/Engine/Source/Programs/Online/OnlineTestsCore/Source/Public/Helpers/Commerce/CommerceQueryOffersByIdHelper.h b/Engine/Source/Programs/Online/OnlineTestsCore/Source/Public/Helpers/Commerce/CommerceQueryOffersByIdHelper.h
new file mode 100644
--- /dev/null
+++ b/Engine/Source/Programs/Online/OnlineTestsCore/Source/Public/Helpers/Commerce/CommerceQueryOffersByIdHelper.h
@@ -0,0 +1,52 @@
+// Copyright Epic Games, Inc. All Rights Reserved.
+
+#pragma once
+
+#include "Online/CommerceCommon.h"
+#include "AsyncTestStep.h"
+#include "OnlineCatchHelper.h"
+
+struct FCommerceQueryOffersByIdHelper : public FAsyncTestStep
+{
+	using ParamsType = UE::Online::FCommerceQueryOffersById::Params;
+	using ResultType = UE::Online::TOnlineResult<UE::Online::FCommerceQueryOffersById>;
+
+	struct FHelperParams
+	{
+		ParamsType* OpParams = nullptr;
+		TOptional<ResultType> ExpectedError;
+	};
+
+	FCommerceQueryOffersByIdHelper(FHelperParams&& InHelperParams)
+		: HelperParams(MoveTemp(InHelperParams))
+	{
+		REQUIRE(HelperParams.OpParams);
+		REQUIRE((!HelperParams.ExpectedError.IsSet() || HelperParams.ExpectedError->IsError()));
+	}
+
+	virtual ~FCommerceQueryOffersByIdHelper() = default;
+
+	virtual void Run(FAsyncStepResult Promise, SubsystemType Services) override
+	{
+		CommerceInterface = Services->GetCommerceInterface();
+		REQUIRE(CommerceInterface);
+
+		CommerceInterface->QueryOffersById(MoveTemp(*HelperParams.OpParams))
+			.OnComplete([this, Promise = MoveTemp(Promise)](const ResultType& Result)
+				{
+					if (!HelperParams.ExpectedError.IsSet())
+					{
+						REQUIRE(Result.IsOk());
+					}
+					else
+					{
+						REQUIRE_OP_EQ(Result, HelperParams.ExpectedError->GetErrorValue());
+					}
+					Promise->SetValue(true);
+				});
+	}
+
+protected:
+	FHelperParams HelperParams;
+	UE::Online::ICommercePtr CommerceInterface = nullptr;
+};
\ No newline at end of file
diff --git a/Engine/Source/Programs/Online/OnlineTestsCore/Source/Public/Helpers/Commerce/CommerceQueryOffersHelper.h b/Engine/Source/Programs/Online/OnlineTestsCore/Source/Public/Helpers/Commerce/CommerceQueryOffersHelper.h
new file mode 100644
--- /dev/null
+++ b/Engine/Source/Programs/Online/OnlineTestsCore/Source/Public/Helpers/Commerce/CommerceQueryOffersHelper.h
@@ -0,0 +1,52 @@
+// Copyright Epic Games, Inc. All Rights Reserved.
+
+#pragma once
+
+#include "Online/CommerceCommon.h"
+#include "AsyncTestStep.h"
+#include "OnlineCatchHelper.h"
+
+struct FCommerceQueryOffersHelper : public FAsyncTestStep
+{
+	using ParamsType = UE::Online::FCommerceQueryOffers::Params;
+	using ResultType = UE::Online::TOnlineResult<UE::Online::FCommerceQueryOffers>;
+
+	struct FHelperParams
+	{
+		ParamsType* OpParams = nullptr;
+		TOptional<ResultType> ExpectedError;
+	};
+
+	FCommerceQueryOffersHelper(FHelperParams&& InHelperParams)
+		: HelperParams(MoveTemp(InHelperParams))
+	{
+		REQUIRE(HelperParams.OpParams);
+		REQUIRE((!HelperParams.ExpectedError.IsSet() || HelperParams.ExpectedError->IsError()));
+	}
+
+	virtual ~FCommerceQueryOffersHelper() = default;
+
+	virtual void Run(FAsyncStepResult Promise, SubsystemType Services) override
+	{
+		CommerceInterface = Services->GetCommerceInterface();
+		REQUIRE(CommerceInterface);
+
+		CommerceInterface->QueryOffers(MoveTemp(*HelperParams.OpParams))
+			.OnComplete([this, Promise = MoveTemp(Promise)](const ResultType& Result)
+				{
+					if (!HelperParams.ExpectedError.IsSet())
+					{
+						REQUIRE(Result.IsOk());
+					}
+					else
+					{
+						REQUIRE_OP_EQ(Result, HelperParams.ExpectedError->GetErrorValue());
+					}
+					Promise->SetValue(true);
+				});
+	}
+
+protected:
+	FHelperParams HelperParams;
+	UE::Online::ICommercePtr CommerceInterface = nullptr;
+};
\ No newline at end of file
