Possible assertions in suite_sessions and returning account to the pool and tests refactoring.

- improve stability of sessions tests
- tests refactoring
- deleting accounts from the template and returning accounts to the TDS pool

diff --git a/Engine/Plugins/Online/OnlineServicesEOS/Source/Private/Online/CommerceEOS.cpp b/Engine/Plugins/Online/OnlineServicesEOS/Source/Private/Online/CommerceEOS.cpp
--- a/Engine/Plugins/Online/OnlineServicesEOS/Source/Private/Online/CommerceEOS.cpp
+++ b/Engine/Plugins/Online/OnlineServicesEOS/Source/Private/Online/CommerceEOS.cpp
@@ -436,12 +436,14 @@ TOnlineAsyncOpHandle<FCommerceRedeemEntitlement> FCommerceEOS::RedeemEntitlement
 		const FCommerceRedeemEntitlement::Params& Params = Op.GetParams();
 		if (!Services.Get<FAuthEOS>()->IsLoggedIn(Params.LocalAccountId))
 		{
+			Promise.SetValue(nullptr);
 			Op.SetError(Errors::NotLoggedIn());
 			return;
 		}
 		EOS_EpicAccountId LocalUserEasId = GetEpicAccountId(Params.LocalAccountId);
 		if (!EOS_EpicAccountId_IsValid(LocalUserEasId))
 		{
+			Promise.SetValue(nullptr);
 			Op.SetError(Errors::NotLoggedIn());
 			return;
 		}
diff --git a/Engine/Source/Programs/Online/OnlineTestsCore/Source/Private/OnlineCatchHelper.cpp b/Engine/Source/Programs/Online/OnlineTestsCore/Source/Private/OnlineCatchHelper.cpp
--- a/Engine/Source/Programs/Online/OnlineTestsCore/Source/Private/OnlineCatchHelper.cpp
+++ b/Engine/Source/Programs/Online/OnlineTestsCore/Source/Private/OnlineCatchHelper.cpp
@@ -19,25 +19,27 @@
 
 // Make sure there are registered input devices for N users and fire
 // OnInputDeviceConnectionChange delegate for interested online service code.
-void EnsureLocalUserCount(uint32 NumUsers)
+void EnsureLocalUserCount(uint32 UserNumToLogin, uint32 NumUsers)
 {
 	TArray<FPlatformUserId> Users;
 	IPlatformInputDeviceMapper::Get().GetAllActiveUsers(Users);
 
-	const uint32 PreviousUserCount = Users.Num();
-	const uint32 NewUserCount = NumUsers > PreviousUserCount ? NumUsers - PreviousUserCount : 0;
+	int32 Index = UserNumToLogin * 10;
+	const int32 EndIndex = NumUsers + Index;
 
-	for (uint32 Index = 0; Index < NewUserCount; ++Index)
+	for (; Index < EndIndex; ++Index)
 	{
-		const uint32 NewUserIndex = PreviousUserCount + Index;
-		IPlatformInputDeviceMapper::Get().Internal_MapInputDeviceToUser(
-			FInputDeviceId::CreateFromInternalId(NewUserIndex),
-			FPlatformMisc::GetPlatformUserForUserIndex(NewUserIndex),
-			EInputDeviceConnectionState::Connected);
+		int32 Result = Users.Find(FPlatformMisc::GetPlatformUserForUserIndex(Index));
+		if (Result < 0)
+		{
+			IPlatformInputDeviceMapper::Get().Internal_MapInputDeviceToUser(
+				FInputDeviceId::CreateFromInternalId(Index),
+				FPlatformMisc::GetPlatformUserForUserIndex(Index),
+				EInputDeviceConnectionState::Connected);
+		}
 	}
 }
 
-
 TArray<TFunction<void()>>* GetGlobalInitalizers()
 {
 	static TArray<TFunction<void()>> gInitalizersToCallInMain;
@@ -77,10 +79,10 @@ SubsystemType OnlineTestBase::GetSubsystem() const
 	return UE::Online::FOnlineServicesRegistry::Get().GetNamedServicesInstance(ServiceType, NAME_None);
 }
 
-bool OnlineTestBase::DeleteAccountsForCurrentTemplate() const
+bool OnlineTestBase::DeleteAccounts(int32 TemplateRequestedNum) const
 {
 #if ONLINETESTS_USEEXTERNAUTH
-	return CustomDeleteAccounts();
+	return CustomDeleteAccounts(TemplateRequestedNum);
 #else // ONLINETESTS_USEEXTERNAUTH
 	return false;
 #endif // ONLINETESTS_USEEXTERNAUTH
@@ -91,10 +93,10 @@ void OnlineTestBase::DestroyCurrentServiceModule() const
 	UE::Online::FOnlineServicesRegistry::Get().DestroyNamedServicesInstance(ServiceType, NAME_None);
 }
 
-bool OnlineTestBase::ResetAccountStatus() const
+bool OnlineTestBase::ResetAccountStatus(int32 TemplateRequestedNum) const
 {
 #if ONLINETESTS_USEEXTERNAUTH
-	return CustomResetAccounts();
+	return CustomResetAccounts(TemplateRequestedNum);
 #else // ONLINETESTS_USEEXTERNAUTH
 	return false;
 #endif // ONLINETESTS_USEEXTERNAUTH
@@ -136,7 +138,7 @@ void OnlineTestBase::UnloadServiceModules()
 	}
 }
 
-FAuthLogin::Params OnlineTestBase::GetIniCredentials(int LocalUserNum) const
+TArray<UE::Online::FAuthLogin::Params> OnlineTestBase::GetIniCredentials(int32 LocalUserNum) const
 {
 	FString LoginCredentialCategory = GetLoginCredentialCategory();
 	TArray<FString> LoginCredentialsRaw;
@@ -146,9 +148,11 @@ FAuthLogin::Params OnlineTestBase::GetIniCredentials(int LocalUserNum) const
 	{
 		UE_LOG_ONLINETESTS(Error, TEXT("Attempted to GetCredentials for more than we have stored! Add more credentials to the DefaultEngine.ini for OnlineTests"));
 		REQUIRE(LocalUserNum >= LoginCredentialsRaw.Num());
-		return FAuthLogin::Params();
+		return TArray<UE::Online::FAuthLogin::Params>();
 	}
 
+	TArray<UE::Online::FAuthLogin::Params> AuthLoginParams;
+
 	TArray<FString> LoginCredentialSplit;
 	LoginCredentialsRaw[LocalUserNum].ParseIntoArray(LoginCredentialSplit, TEXT(","), true);
 
@@ -165,13 +169,15 @@ FAuthLogin::Params OnlineTestBase::GetIniCredentials(int LocalUserNum) const
 	}
 	Params.PlatformUserId = FPlatformMisc::GetPlatformUserForUserIndex(LocalUserNum);
 
-	return Params;
+	AuthLoginParams.Add(Params);
+
+	return AuthLoginParams;
 }
 
-FAuthLogin::Params OnlineTestBase::GetCredentials(int LocalUserNum) const
+TArray<UE::Online::FAuthLogin::Params> OnlineTestBase::GetCredentials(int32 LocalUserNum, int32 NumUsers) const
 {
 #if ONLINETESTS_USEEXTERNAUTH
-	return CustomCredentials(LocalUserNum);
+	return CustomCredentials(LocalUserNum, NumUsers);
 #else // ONLINETESTS_USEEXTERNAUTH
 	return GetIniCredentials(LocalUserNum);
 #endif // ONLINETESTS_USEEXTERNAUTH
@@ -182,50 +188,6 @@ FString OnlineTestBase::GetLoginCredentialCategory() const
 	return FString::Printf(TEXT("LoginCredentials %s"), *Service);
 }
 
-FTestPipeline& OnlineTestBase::GetLoginPipelineArray(std::initializer_list<std::reference_wrapper<FAccountId>> AccountIdsArr) const
-{
-	const int NumUsersToLogin = AccountIdsArr.size();
-	GetLoginPipeline(NumUsersToLogin);
-
-	// Perform login so we can bulk assign users in the next step.
-	RunToCompletion(false);
-
-	int32 LocalUserNum = 0;
-	for (FAccountId& AccountId : AccountIdsArr)
-	{
-		AssignLoginUsers(LocalUserNum++, AccountId);
-	}
-
-	// Return a fresh pipeline so the logins added by GetLoginPipeline don't execute again.
-	Pipeline = MakeShared<FTestPipeline>(Driver.MakePipeline());
-	return *Pipeline;
-}
-
-FTestPipeline& OnlineTestBase::GetLoginPipeline(FAccountId& AccountId) const
-{
-	return GetLoginPipelineArray({ AccountId });
-}
-
-FTestPipeline& OnlineTestBase::GetLoginPipeline(FAccountId& AccountId, FAccountId& AccountId2) const
-{
-	return GetLoginPipelineArray({ AccountId, AccountId2 });
-}
-
-FTestPipeline& OnlineTestBase::GetLoginPipeline(FAccountId& AccountId, FAccountId& AccountId2, FAccountId& AccountId3) const
-{
-	return GetLoginPipelineArray({ AccountId, AccountId2, AccountId3 });
-}
-
-FTestPipeline& OnlineTestBase::GetLoginPipeline(FAccountId& AccountId, FAccountId& AccountId2, FAccountId& AccountId3, FAccountId& AccountId4) const
-{
-	return GetLoginPipelineArray({ AccountId, AccountId2, AccountId3, AccountId4 });
-}
-
-FTestPipeline& OnlineTestBase::GetLoginPipeline(FAccountId& AccountId, FAccountId& AccountId2, FAccountId& AccountId3, FAccountId& AccountId4, FAccountId& AccountId5) const
-{
-	return GetLoginPipelineArray({ AccountId, AccountId2, AccountId3, AccountId4, AccountId5 });
-}
-
 void OnlineTestBase::AssignLoginUsers(int32 LocalUserId, FAccountId& OutAccountId) const
 {
 	SubsystemType OnlineSubsystem = GetSubsystem();
@@ -235,23 +197,35 @@ void OnlineTestBase::AssignLoginUsers(int32 LocalUserId, FAccountId& OutAccountI
 	OutAccountId = UserId.TryGetOkValue()->AccountInfo->AccountId;
 }
 
-FTestPipeline& OnlineTestBase::GetLoginPipeline(uint32 NumUsersToLogin) const
+FTestPipeline& OnlineTestBase::GetLoginPipeline(std::initializer_list<std::reference_wrapper<FAccountId>> AccountIds) const
 {
-	REQUIRE(NumLocalUsers == -1); // Don't call GetLoginPipeline more than once per test
-	NumLocalUsers = NumUsersToLogin;
+	return GetLoginPipeline(0, AccountIds);
+}
+
+FTestPipeline& OnlineTestBase::GetLoginPipeline(uint32 UserNumToLogin, std::initializer_list<std::reference_wrapper<FAccountId>> AccountIds) const
+{
+	// Don't call GetLoginPipeline more than once per test
+	REQUIRE(NumLocalUsers == -1);
+	NumLocalUsers = AccountIds.size();
 
 	bool bUseAutoLogin = false;
 	bool bUseImplicitLogin = false;
 	FString LoginCredentialCategory = GetLoginCredentialCategory();
 	GConfig->GetBool(*LoginCredentialCategory, TEXT("UseAutoLogin"), bUseAutoLogin, GEngineIni);
 	GConfig->GetBool(*LoginCredentialCategory, TEXT("UseImplicitLogin"), bUseImplicitLogin, GEngineIni);
-
+	
 	// Make sure input delegates are fired for adding the required user count.
-	EnsureLocalUserCount(NumUsersToLogin);
+	EnsureLocalUserCount(UserNumToLogin, NumLocalUsers);
 
 	if (bUseImplicitLogin)
 	{
 		// Users are expected to already be valid.
+		int32 Index = 0;
+		for (FAccountId& AccountId : AccountIds)
+		{
+			AssignLoginUsers(Index, AccountId);
+			++Index;
+		}
 	}
 	else if (bUseAutoLogin)
 	{
@@ -259,26 +233,28 @@ FTestPipeline& OnlineTestBase::GetLoginPipeline(uint32 NumUsersToLogin) const
 		// NumLocalUsers = 1;
 		// Pipeline.EmplaceStep<FAuthAutoLoginStep>(0);
 	}
-	else
+	else if (NumLocalUsers > 0)
 	{
-		for (uint32 i = 0; i < NumUsersToLogin; i++)
+		UserNumToLogout = UserNumToLogin;
+
+		TArray<UE::Online::FAuthLogin::Params> AuthLoginParams = GetCredentials(UserNumToLogin, NumLocalUsers);
+		for (UE::Online::FAuthLogin::Params AuthLoginParam : AuthLoginParams)
 		{
-			Pipeline->EmplaceStep<FAuthLoginStep>(GetCredentials(i));
+			Pipeline->EmplaceStep<FAuthLoginStep>(MoveTemp(AuthLoginParam));
 		}
-	}
 
-	return *Pipeline;
-}
+		// Perform login so we can bulk assign users in the next step.
+		RunToCompletion(false);
+
+		int32 Index = 0;
+		for (FAccountId& AccountId : AccountIds)
+		{
+			UE_LOG_ONLINETESTS(Verbose, TEXT("Account: %s, used InrernalId: %d"), *AuthLoginParams[Index].CredentialsId, AuthLoginParams[Index].PlatformUserId.GetInternalId());
+			AssignLoginUsers(AuthLoginParams[Index].PlatformUserId.GetInternalId(), AccountId);
+			++Index;
+		}
+	}
 
-FTestPipeline& OnlineTestBase::GetLoginPipeline(uint32 UserNumToLogin, FAccountId& OutAccountId) const
-{
-	REQUIRE(NumLocalUsers == -1); // Don't call GetLoginPipeline more than once per test
-	// Make sure input delegates are fired for adding the required user count.
-	//EnsureLocalUserCount(UserNumToLogin);
-	Pipeline->EmplaceStep<FAuthLoginStep>(GetCredentials(UserNumToLogin));
-	// Perform login so we can bulk assign users in the next step.
-	RunToCompletion(false);
-	AssignLoginUsers(UserNumToLogin, OutAccountId);
 	// Return a fresh pipeline so the logins added by GetLoginPipeline don't execute again.
 	Pipeline = MakeShared<FTestPipeline>(Driver.MakePipeline());
 	return *Pipeline;
@@ -286,10 +262,10 @@ FTestPipeline& OnlineTestBase::GetLoginPipeline(uint32 UserNumToLogin, FAccountI
 
 FTestPipeline& OnlineTestBase::GetPipeline() const
 {
-	return GetLoginPipeline(0);
+	return GetLoginPipeline(0, {});
 }
 
-void OnlineTestBase::RunToCompletion(bool bLogout, const TOptional<int32> UserNumToLogout) const
+void OnlineTestBase::RunToCompletion(bool bLogout) const
 {
 	bool bUseAutoLogin = false;
 	bool bUseImplicitLogin = false;
@@ -308,15 +284,11 @@ void OnlineTestBase::RunToCompletion(bool bLogout, const TOptional<int32> UserNu
 			// NumLocalUsers = 1;
 			// Pipeline.EmplaceStep<FAuthAutoLoginStep>(0);
 		}
-		else if (UserNumToLogout.IsSet())
+		else if(NumLocalUsers > 0)
 		{
-			Pipeline->EmplaceStep<FAuthLogoutStep>(FPlatformMisc::GetPlatformUserForUserIndex(UserNumToLogout.GetValue()));
-		}
-		else 
-		{
-			for (uint32 i = 0; i < NumLocalUsers; i++)
+			for(UE::Online::FAuthLogin::Params AuthLoginParam : GetCredentials(UserNumToLogout, NumLocalUsers))
 			{
-				Pipeline->EmplaceStep<FAuthLogoutStep>(FPlatformMisc::GetPlatformUserForUserIndex(i));
+				Pipeline->EmplaceStep<FAuthLogoutStep>(AuthLoginParam.PlatformUserId);
 			}
 		}
 	}
@@ -371,7 +343,6 @@ TArray<OnlineAutoReg::FApplicableServicesConfig> OnlineAutoReg::GetApplicableSer
 	return ServicesConfig;
 }
 
-
 bool OnlineAutoReg::CheckAllTagsIsIn(const TArray<FString>& TestTags, const TArray<FString>& InputTags)
 {
 	if (InputTags.Num() == 0)
diff --git a/Engine/Source/Programs/Online/OnlineTestsCore/Source/Private/Tests/Achievements/DisplayAchievementUITests.cpp b/Engine/Source/Programs/Online/OnlineTestsCore/Source/Private/Tests/Achievements/DisplayAchievementUITests.cpp
--- a/Engine/Source/Programs/Online/OnlineTestsCore/Source/Private/Tests/Achievements/DisplayAchievementUITests.cpp
+++ b/Engine/Source/Programs/Online/OnlineTestsCore/Source/Private/Tests/Achievements/DisplayAchievementUITests.cpp
@@ -6,26 +6,23 @@
 #include "Helpers/LambdaStep.h"
 #include "OnlineCatchHelper.h"
 
-
 #define ACHIEVEMENTS_TAG "[suite_achievements]"
 #define EG_ACHIEVEMENTS_DISPLAYACHIEVEMENTUI_TAG ACHIEVEMENTS_TAG "[displayachievementui][.NULL]"
 
 #define ACHIEVEMENTS_TEST_CASE(x, ...) ONLINE_TEST_CASE(x, ACHIEVEMENTS_TAG __VA_ARGS__)
 
-
 ACHIEVEMENTS_TEST_CASE("Display Achievement UI (Invalid User)", EG_ACHIEVEMENTS_DISPLAYACHIEVEMENTUI_TAG)
 {
-	int32 NumUsersToImplicitLogin = 0;
-
-	GetLoginPipeline(NumUsersToImplicitLogin).EmplaceLambda([](SubsystemType OnlineSubsystem)
-		{
-			UE::Online::FDisplayAchievementUI::Params OpParams;
+	GetPipeline()
+		.EmplaceLambda([](SubsystemType OnlineSubsystem)
+			{
+				UE::Online::FDisplayAchievementUI::Params OpParams;
 
-			UE::Online::IAchievementsPtr AchievementsInterface = OnlineSubsystem->GetAchievementsInterface();
-			UE::Online::TOnlineResult<UE::Online::FDisplayAchievementUI> Result = AchievementsInterface->DisplayAchievementUI(MoveTemp(OpParams));
-			REQUIRE(Result.IsError());
-			CHECK(Result.GetErrorValue() == UE::Online::Errors::InvalidUser());
-		});
+				UE::Online::IAchievementsPtr AchievementsInterface = OnlineSubsystem->GetAchievementsInterface();
+				UE::Online::TOnlineResult<UE::Online::FDisplayAchievementUI> Result = AchievementsInterface->DisplayAchievementUI(MoveTemp(OpParams));
+				REQUIRE(Result.IsError());
+				CHECK(Result.GetErrorValue() == UE::Online::Errors::InvalidUser());
+			});
 
 	RunToCompletion();
 }
@@ -37,7 +34,7 @@ ACHIEVEMENTS_TEST_CASE("Display Achievement UI (Invalid State)",  EG_ACHIEVEMENT
 	FAccountId AccountId;
 	FQueryAchievementDefinitionsHelper::FHelperParams QueryDefinitionsHelperParams;
 
-	FTestPipeline& LoginPipeline = GetLoginPipeline(AccountId);
+	FTestPipeline& LoginPipeline = GetLoginPipeline({ AccountId });
 
 	QueryDefinitionsHelperParams.OpParams.LocalAccountId = AccountId;
 
@@ -63,7 +60,7 @@ ACHIEVEMENTS_TEST_CASE("Display Achievement UI (Not Found)", EG_ACHIEVEMENTS_DIS
 	FQueryAchievementDefinitionsHelper::FHelperParams QueryDefinitionsHelperParams;
 	FQueryAchievementStatesHelper::FHelperParams QueryStatesHelperParams;
 	
-	FTestPipeline& LoginPipeline = GetLoginPipeline(AccountId);
+	FTestPipeline& LoginPipeline = GetLoginPipeline({ AccountId });
 
 	QueryDefinitionsHelperParams.OpParams.LocalAccountId = AccountId;
 	QueryStatesHelperParams.OpParams.LocalAccountId = AccountId;
@@ -92,7 +89,7 @@ ACHIEVEMENTS_TEST_CASE("Display Achievement UI (Success)", EG_ACHIEVEMENTS_DISPL
 	FQueryAchievementDefinitionsHelper::FHelperParams QueryDefinitionsHelperParams;
 	FQueryAchievementStatesHelper::FHelperParams QueryStatesHelperParams;
 
-	FTestPipeline& LoginPipeline = GetLoginPipeline(AccountId);
+	FTestPipeline& LoginPipeline = GetLoginPipeline({ AccountId });
 
 	QueryDefinitionsHelperParams.OpParams.LocalAccountId = AccountId;
 	QueryStatesHelperParams.OpParams.LocalAccountId = AccountId;
diff --git a/Engine/Source/Programs/Online/OnlineTestsCore/Source/Private/Tests/Achievements/GetAchievementDefinitionTests.cpp b/Engine/Source/Programs/Online/OnlineTestsCore/Source/Private/Tests/Achievements/GetAchievementDefinitionTests.cpp
--- a/Engine/Source/Programs/Online/OnlineTestsCore/Source/Private/Tests/Achievements/GetAchievementDefinitionTests.cpp
+++ b/Engine/Source/Programs/Online/OnlineTestsCore/Source/Private/Tests/Achievements/GetAchievementDefinitionTests.cpp
@@ -5,26 +5,23 @@
 #include "Helpers/LambdaStep.h"
 #include "OnlineCatchHelper.h"
 
-
 #define ACHIEVEMENTS_TAG "[suite_achievements]"
 #define EG_ACHIEVEMENTS_GETACHIEVEMENTDEFINITION_TAG ACHIEVEMENTS_TAG "[getachievementdefinition]"
 
 #define ACHIEVEMENTS_TEST_CASE(x, ...) ONLINE_TEST_CASE(x, ACHIEVEMENTS_TAG __VA_ARGS__)
 
-
 ACHIEVEMENTS_TEST_CASE("Get Achievement Definition (Invalid User)", EG_ACHIEVEMENTS_GETACHIEVEMENTDEFINITION_TAG)
 {
-	int32 NumUsersToImplicitLogin = 0;
-
-	GetLoginPipeline(NumUsersToImplicitLogin).EmplaceLambda([](SubsystemType OnlineSubsystem)
-		{
-			UE::Online::FGetAchievementDefinition::Params OpParams;
+	GetPipeline()
+		.EmplaceLambda([](SubsystemType OnlineSubsystem)
+			{
+				UE::Online::FGetAchievementDefinition::Params OpParams;
 
-			UE::Online::IAchievementsPtr AchievementsInterface = OnlineSubsystem->GetAchievementsInterface();
-			UE::Online::TOnlineResult<UE::Online::FGetAchievementDefinition> Result = AchievementsInterface->GetAchievementDefinition(MoveTemp(OpParams));
-			REQUIRE(Result.IsError());
-			CHECK(Result.GetErrorValue() == UE::Online::Errors::InvalidUser());
-		});
+				UE::Online::IAchievementsPtr AchievementsInterface = OnlineSubsystem->GetAchievementsInterface();
+				UE::Online::TOnlineResult<UE::Online::FGetAchievementDefinition> Result = AchievementsInterface->GetAchievementDefinition(MoveTemp(OpParams));
+				REQUIRE(Result.IsError());
+				CHECK(Result.GetErrorValue() == UE::Online::Errors::InvalidUser());
+			});
 
 	RunToCompletion();
 }
@@ -35,7 +32,7 @@ ACHIEVEMENTS_TEST_CASE("Get Achievement Definition (Invalid State)", EG_ACHIEVEM
 
 	FAccountId AccountId;
 
-	FTestPipeline& LoginPipeline = GetLoginPipeline(AccountId);
+	FTestPipeline& LoginPipeline = GetLoginPipeline({ AccountId });
 
 	UE::Online::FGetAchievementDefinition::Params GetAchievementDefinitionParams;
 	GetAchievementDefinitionParams.LocalAccountId = AccountId;
@@ -55,7 +52,7 @@ ACHIEVEMENTS_TEST_CASE("Get Achievement Definition (Not Found)", EG_ACHIEVEMENTS
 {
 	FAccountId AccountId;
 
-	FTestPipeline& LoginPipeline = GetLoginPipeline(AccountId);
+	FTestPipeline& LoginPipeline = GetLoginPipeline({ AccountId });
 
 	FQueryAchievementDefinitionsHelper::FHelperParams QueryDefinitionsHelperParams;
 	QueryDefinitionsHelperParams.OpParams.LocalAccountId = AccountId;
@@ -83,7 +80,7 @@ ACHIEVEMENTS_TEST_CASE("Get Achievement Definitions (Success)", EG_ACHIEVEMENTS_
 	FAccountId AccountId;
 	TArray<FString> AchievementIds;
 
-	FTestPipeline& LoginPipeline = GetLoginPipeline(AccountId);
+	FTestPipeline& LoginPipeline = GetLoginPipeline({ AccountId });
 
 	FQueryAchievementDefinitionsHelper::FHelperParams QueryDefinitionsHelperParams;
 	QueryDefinitionsHelperParams.OpParams.LocalAccountId = AccountId;
diff --git a/Engine/Source/Programs/Online/OnlineTestsCore/Source/Private/Tests/Achievements/GetAchievementIdsTests.cpp b/Engine/Source/Programs/Online/OnlineTestsCore/Source/Private/Tests/Achievements/GetAchievementIdsTests.cpp
--- a/Engine/Source/Programs/Online/OnlineTestsCore/Source/Private/Tests/Achievements/GetAchievementIdsTests.cpp
+++ b/Engine/Source/Programs/Online/OnlineTestsCore/Source/Private/Tests/Achievements/GetAchievementIdsTests.cpp
@@ -5,26 +5,23 @@
 #include "Helpers/LambdaStep.h"
 #include "OnlineCatchHelper.h"
 
-
 #define ACHIEVEMENTS_TAG "[suite_achievements]"
 #define EG_ACHIEVEMENTS_GETACHIEVEMENTIDS_TAG ACHIEVEMENTS_TAG "[getachievementids]"
 
 #define ACHIEVEMENTS_TEST_CASE(x, ...) ONLINE_TEST_CASE(x, ACHIEVEMENTS_TAG __VA_ARGS__)
 
-
 ACHIEVEMENTS_TEST_CASE("Get Achievement Ids (Invalid User)", EG_ACHIEVEMENTS_GETACHIEVEMENTIDS_TAG)
 {
-	int32 NumUsersToImplicitLogin = 0;
-
-	GetLoginPipeline(NumUsersToImplicitLogin).EmplaceLambda([](SubsystemType OnlineSubsystem)
-		{
-			UE::Online::FGetAchievementIds::Params OpParams;
+	GetPipeline()
+		.EmplaceLambda([](SubsystemType OnlineSubsystem)
+			{
+				UE::Online::FGetAchievementIds::Params OpParams;
 
-			UE::Online::IAchievementsPtr AchievementsInterface = OnlineSubsystem->GetAchievementsInterface();
-			UE::Online::TOnlineResult<UE::Online::FGetAchievementIds> Result = AchievementsInterface->GetAchievementIds(MoveTemp(OpParams));
-			REQUIRE(Result.IsError());
-			CHECK(Result.GetErrorValue() == UE::Online::Errors::InvalidUser());
-		});
+				UE::Online::IAchievementsPtr AchievementsInterface = OnlineSubsystem->GetAchievementsInterface();
+				UE::Online::TOnlineResult<UE::Online::FGetAchievementIds> Result = AchievementsInterface->GetAchievementIds(MoveTemp(OpParams));
+				REQUIRE(Result.IsError());
+				CHECK(Result.GetErrorValue() == UE::Online::Errors::InvalidUser());
+			});
 
 	RunToCompletion();
 }
@@ -35,7 +32,7 @@ ACHIEVEMENTS_TEST_CASE("Get Achievement Ids (Invalid State)", EG_ACHIEVEMENTS_GE
 
 	FAccountId AccountId;
 
-	auto& LoginPipeline = GetLoginPipeline(AccountId);
+	auto& LoginPipeline = GetLoginPipeline({ AccountId });
 
 	UE::Online::FGetAchievementIds::Params OpParams;
 	OpParams.LocalAccountId = AccountId;
@@ -55,7 +52,7 @@ ACHIEVEMENTS_TEST_CASE("Get Achievement Ids (Success)", EG_ACHIEVEMENTS_GETACHIE
 {
 	FAccountId AccountId;
 
-	auto& LoginPipeline = GetLoginPipeline(AccountId);
+	auto& LoginPipeline = GetLoginPipeline({ AccountId });
 
 	FQueryAchievementDefinitionsHelper::FHelperParams QueryDefinitionsHelperParams;
 	QueryDefinitionsHelperParams.OpParams.LocalAccountId = AccountId;
diff --git a/Engine/Source/Programs/Online/OnlineTestsCore/Source/Private/Tests/Achievements/GetAchievementStateTests.cpp b/Engine/Source/Programs/Online/OnlineTestsCore/Source/Private/Tests/Achievements/GetAchievementStateTests.cpp
--- a/Engine/Source/Programs/Online/OnlineTestsCore/Source/Private/Tests/Achievements/GetAchievementStateTests.cpp
+++ b/Engine/Source/Programs/Online/OnlineTestsCore/Source/Private/Tests/Achievements/GetAchievementStateTests.cpp
@@ -6,20 +6,18 @@
 #include "Helpers/LambdaStep.h"
 #include "OnlineCatchHelper.h"
 
-
 #define ACHIEVEMENTS_TAG "[suite_achievements]"
 #define EG_ACHIEVEMENTS_GETACHIEVEMENTSTATE_TAG ACHIEVEMENTS_TAG "[getachievementstate]"
 
 #define ACHIEVEMENTS_TEST_CASE(x, ...) ONLINE_TEST_CASE(x, ACHIEVEMENTS_TAG __VA_ARGS__)
 
-
 ACHIEVEMENTS_TEST_CASE("Get Achievement State (Invalid State)", EG_ACHIEVEMENTS_GETACHIEVEMENTSTATE_TAG)
 {
 	DestroyCurrentServiceModule();
 
 	FAccountId AccountId;
 
-	FTestPipeline& LoginPipeline = GetLoginPipeline(AccountId);
+	FTestPipeline& LoginPipeline = GetLoginPipeline({ AccountId });
 
 	FQueryAchievementDefinitionsHelper::FHelperParams QueryDefinitionsHelperParams;
 	QueryDefinitionsHelperParams.OpParams.LocalAccountId = AccountId;
@@ -44,7 +42,7 @@ ACHIEVEMENTS_TEST_CASE("Get Achievement State (Not Found)", EG_ACHIEVEMENTS_GETA
 {
 	FAccountId AccountId;
 
-	FTestPipeline& LoginPipeline = GetLoginPipeline(AccountId);
+	FTestPipeline& LoginPipeline = GetLoginPipeline({ AccountId });
 
 	FQueryAchievementDefinitionsHelper::FHelperParams QueryDefinitionsHelperParams;
 	QueryDefinitionsHelperParams.OpParams.LocalAccountId = AccountId;
@@ -73,7 +71,7 @@ ACHIEVEMENTS_TEST_CASE("Get Achievement State (Success)", EG_ACHIEVEMENTS_GETACH
 {
 	FAccountId AccountId;
 
-	FTestPipeline& LoginPipeline = GetLoginPipeline(AccountId);
+	FTestPipeline& LoginPipeline = GetLoginPipeline({ AccountId });
 
 	FQueryAchievementDefinitionsHelper::FHelperParams QueryDefinitionsHelperParams;
 	QueryDefinitionsHelperParams.OpParams.LocalAccountId = AccountId;
diff --git a/Engine/Source/Programs/Online/OnlineTestsCore/Source/Private/Tests/Achievements/QueryAchievementDefinitionsTests.cpp b/Engine/Source/Programs/Online/OnlineTestsCore/Source/Private/Tests/Achievements/QueryAchievementDefinitionsTests.cpp
--- a/Engine/Source/Programs/Online/OnlineTestsCore/Source/Private/Tests/Achievements/QueryAchievementDefinitionsTests.cpp
+++ b/Engine/Source/Programs/Online/OnlineTestsCore/Source/Private/Tests/Achievements/QueryAchievementDefinitionsTests.cpp
@@ -5,13 +5,11 @@
 #include "Helpers/LambdaStep.h"
 #include "OnlineCatchHelper.h"
 
-
 #define ACHIEVEMENTS_TAG "[suite_achievements]"
 #define EG_ACHIEVEMENTS_QUERYACHIEVEMENTDEFINITIONS_TAG ACHIEVEMENTS_TAG "[queryachievementdefinitions]"
 
 #define ACHIEVEMENTS_TEST_CASE(x, ...) ONLINE_TEST_CASE(x, ACHIEVEMENTS_TAG __VA_ARGS__)
 
-
 ACHIEVEMENTS_TEST_CASE("Query Definitions (Invalid User)", EG_ACHIEVEMENTS_QUERYACHIEVEMENTDEFINITIONS_TAG)
 {
 	FAccountId AccountId;
@@ -19,7 +17,7 @@ ACHIEVEMENTS_TEST_CASE("Query Definitions (Invalid User)", EG_ACHIEVEMENTS_QUERY
 	FQueryAchievementDefinitionsHelper::FHelperParams QueryDefinitionsHelperParams;
 	QueryDefinitionsHelperParams.ExpectedError = FQueryAchievementDefinitionsHelper::ResultType(UE::Online::Errors::InvalidUser());
 
-	GetLoginPipeline(AccountId)
+	GetLoginPipeline({ AccountId })
 		.EmplaceStep<FQueryAchievementDefinitionsHelper>(MoveTemp(QueryDefinitionsHelperParams));
 
 	RunToCompletion();
@@ -29,7 +27,7 @@ ACHIEVEMENTS_TEST_CASE("Query Definitions (Success)", EG_ACHIEVEMENTS_QUERYACHIE
 {
 	FAccountId AccountId;
 
-	FTestPipeline& LoginPipeline = GetLoginPipeline(AccountId);
+	FTestPipeline& LoginPipeline = GetLoginPipeline({ AccountId });
 
 	FQueryAchievementDefinitionsHelper::FHelperParams QueryDefinitionsHelperParams;
 	QueryDefinitionsHelperParams.OpParams.LocalAccountId = AccountId;
diff --git a/Engine/Source/Programs/Online/OnlineTestsCore/Source/Private/Tests/Achievements/QueryAchievementStatesTests.cpp b/Engine/Source/Programs/Online/OnlineTestsCore/Source/Private/Tests/Achievements/QueryAchievementStatesTests.cpp
--- a/Engine/Source/Programs/Online/OnlineTestsCore/Source/Private/Tests/Achievements/QueryAchievementStatesTests.cpp
+++ b/Engine/Source/Programs/Online/OnlineTestsCore/Source/Private/Tests/Achievements/QueryAchievementStatesTests.cpp
@@ -7,13 +7,11 @@
 #include "Helpers/LambdaStep.h"
 #include "OnlineCatchHelper.h"
 
-
 #define ACHIEVEMENTS_TAG "[suite_achievements]"
 #define EG_ACHIEVEMENTS_QUERYACHIEVEMENTSTATES_TAG ACHIEVEMENTS_TAG "[queryachievementstates]"
 
 #define ACHIEVEMENTS_TEST_CASE(x, ...) ONLINE_TEST_CASE(x, ACHIEVEMENTS_TAG __VA_ARGS__)
 
-
 ACHIEVEMENTS_TEST_CASE("Query Achievement States (Invalid User)", EG_ACHIEVEMENTS_QUERYACHIEVEMENTSTATES_TAG)
 {
 	FAccountId AccountId;
@@ -21,7 +19,7 @@ ACHIEVEMENTS_TEST_CASE("Query Achievement States (Invalid User)", EG_ACHIEVEMENT
 	FQueryAchievementStatesHelper::FHelperParams QueryStatesHelperParams;
 	QueryStatesHelperParams.ExpectedError = FQueryAchievementStatesHelper::ResultType(UE::Online::Errors::InvalidUser());
 
-	GetLoginPipeline(AccountId)
+	GetLoginPipeline({ AccountId })
 		.EmplaceStep<FQueryAchievementStatesHelper>(MoveTemp(QueryStatesHelperParams));
 
 	RunToCompletion();
@@ -33,7 +31,7 @@ ACHIEVEMENTS_TEST_CASE("Query Achievement States (Invalid State)", EG_ACHIEVEMEN
 
 	FAccountId AccountId;
 
-	FTestPipeline& LoginPipeline = GetLoginPipeline(AccountId);
+	FTestPipeline& LoginPipeline = GetLoginPipeline({ AccountId });
 
 	FQueryAchievementStatesHelper::FHelperParams QueryStatesHelperParams;
 	QueryStatesHelperParams.ExpectedError = FQueryAchievementStatesHelper::ResultType(UE::Online::Errors::InvalidState());
@@ -48,7 +46,7 @@ ACHIEVEMENTS_TEST_CASE("Query Achievement States (Success)", EG_ACHIEVEMENTS_QUE
 {
 	FAccountId AccountId;
 
-	FTestPipeline& LoginPipeline = GetLoginPipeline(AccountId);
+	FTestPipeline& LoginPipeline = GetLoginPipeline({ AccountId });
 	
 	FQueryAchievementDefinitionsHelper::FHelperParams QueryDefinitionsHelperParams;
 	QueryDefinitionsHelperParams.OpParams.LocalAccountId = AccountId;
@@ -66,7 +64,7 @@ ACHIEVEMENTS_TEST_CASE("Get Achievement State (Invalid User)", EG_ACHIEVEMENTS_Q
 {
 	FAccountId AccountId;
 
-	GetLoginPipeline(AccountId)
+	GetLoginPipeline({ AccountId })
 		.EmplaceLambda([](SubsystemType OnlineSubsystem)
 			{
 				UE::Online::FGetAchievementState::Params Params;
diff --git a/Engine/Source/Programs/Online/OnlineTestsCore/Source/Private/Tests/Achievements/UnlockAchievementsTests.cpp b/Engine/Source/Programs/Online/OnlineTestsCore/Source/Private/Tests/Achievements/UnlockAchievementsTests.cpp
--- a/Engine/Source/Programs/Online/OnlineTestsCore/Source/Private/Tests/Achievements/UnlockAchievementsTests.cpp
+++ b/Engine/Source/Programs/Online/OnlineTestsCore/Source/Private/Tests/Achievements/UnlockAchievementsTests.cpp
@@ -12,14 +12,12 @@
 #include "Online/Stats.h"
 #include "OnlineCatchHelper.h"
 
-
 #define ACHIEVEMENTS_TAG "[suite_achievements]"
 #define EG_ACHIEVEMENTS_UNLOCKACHIEVEMENTS_TAG ACHIEVEMENTS_TAG "[unlockachievements]"
 #define EG_ACHIEVEMENTS_UNLOCKACHIEVEMENTS_TITLEMANAGED_TAG EG_ACHIEVEMENTS_UNLOCKACHIEVEMENTS_TAG "[.NULL]"
 
 #define ACHIEVEMENTS_TEST_CASE(x, ...) ONLINE_TEST_CASE(x, ACHIEVEMENTS_TAG __VA_ARGS__)
 
-
 ACHIEVEMENTS_TEST_CASE("Unlock Achievements (Invalid User)", EG_ACHIEVEMENTS_UNLOCKACHIEVEMENTS_TAG)
 {
 	FAccountId AccountId;
@@ -28,7 +26,7 @@ ACHIEVEMENTS_TEST_CASE("Unlock Achievements (Invalid User)", EG_ACHIEVEMENTS_UNL
 	UnlockAchievementsHelperParams.OpParams.AchievementIds = { TEXT("test_unlockachievements") };
 	UnlockAchievementsHelperParams.ExpectedError = FUnlockAchievementsHelper::ResultType(UE::Online::Errors::InvalidUser());
 
-	GetLoginPipeline(AccountId)
+	GetLoginPipeline({ AccountId })
 		.EmplaceStep<FUnlockAchievementsHelper>(MoveTemp(UnlockAchievementsHelperParams));
 
 	RunToCompletion();
@@ -40,7 +38,7 @@ ACHIEVEMENTS_TEST_CASE("Unlock Achievements (Invalid State)", EG_ACHIEVEMENTS_UN
 
 	FAccountId AccountId;
 
-	FTestPipeline& LoginPipeline = GetLoginPipeline(AccountId);
+	FTestPipeline& LoginPipeline = GetLoginPipeline({ AccountId });
 
 	FQueryAchievementDefinitionsHelper::FHelperParams QueryDefinitionsHelperParams;
 	QueryDefinitionsHelperParams.OpParams.LocalAccountId = AccountId;
@@ -61,7 +59,7 @@ ACHIEVEMENTS_TEST_CASE("Unlock Achievements (Invalid Params)", EG_ACHIEVEMENTS_U
 {
 	FAccountId AccountId;
 
-	FTestPipeline& LoginPipeline = GetLoginPipeline(AccountId);
+	FTestPipeline& LoginPipeline = GetLoginPipeline({ AccountId });
 
 	FQueryAchievementDefinitionsHelper::FHelperParams QueryDefinitionsHelperParams;
 	QueryDefinitionsHelperParams.OpParams.LocalAccountId = AccountId;
@@ -84,7 +82,7 @@ ACHIEVEMENTS_TEST_CASE("Unlock Achievements (Not Found)", EG_ACHIEVEMENTS_UNLOCK
 {
 	FAccountId AccountId;
 
-	FTestPipeline& LoginPipeline = GetLoginPipeline(AccountId);
+	FTestPipeline& LoginPipeline = GetLoginPipeline({ AccountId });
 
 	FQueryAchievementDefinitionsHelper::FHelperParams QueryDefinitionsHelperParams;
 	QueryDefinitionsHelperParams.OpParams.LocalAccountId = AccountId;
@@ -108,12 +106,9 @@ ACHIEVEMENTS_TEST_CASE("Unlock Achievements (Already Unlocked)", EG_ACHIEVEMENTS
 {
 	DestroyCurrentServiceModule();
 
-	//Delete all accounts as a form of reset, this way future tests are always performed on fresh accounts
-	DeleteAccountsForCurrentTemplate();
-
 	FAccountId AccountId;
 
-	FTestPipeline& LoginPipeline = GetLoginPipeline(AccountId);
+	FTestPipeline& LoginPipeline = GetLoginPipeline({ AccountId });
 
 	FQueryAchievementDefinitionsHelper::FHelperParams QueryDefinitionsHelperParams;
 	QueryDefinitionsHelperParams.OpParams.LocalAccountId = AccountId;
@@ -144,7 +139,7 @@ ACHIEVEMENTS_TEST_CASE("Unlock Achievements (Success)", EG_ACHIEVEMENTS_UNLOCKAC
 
 	FAccountId AccountId;
 
-	FTestPipeline& LoginPipeline = GetLoginPipeline(AccountId);
+	FTestPipeline& LoginPipeline = GetLoginPipeline({ AccountId });
 
 	FQueryAchievementDefinitionsHelper::FHelperParams QueryDefinitionsHelperParams;
 	QueryDefinitionsHelperParams.OpParams.LocalAccountId = AccountId;
@@ -169,7 +164,7 @@ ACHIEVEMENTS_TEST_CASE("Unlock Title-Managed Achievements By Stats", EG_ACHIEVEM
 
 	FAccountId AccountId;
 
-	FTestPipeline& LoginPipeline = GetLoginPipeline(AccountId);
+	FTestPipeline& LoginPipeline = GetLoginPipeline({ AccountId });
 
 	FString AchievementId("FAKE_ACHIEVEMENT");
 
@@ -209,7 +204,7 @@ ACHIEVEMENTS_TEST_CASE("Unlock Title-Managed Achievements By Stats With Multiple
 
 	FAccountId AccountId;
 
-	FTestPipeline& LoginPipeline = GetLoginPipeline(AccountId);
+	FTestPipeline& LoginPipeline = GetLoginPipeline({ AccountId });
 
 	FString AchievementId("test_unlockachievements");
 
diff --git a/Engine/Source/Programs/Online/OnlineTestsCore/Source/Private/Tests/Auth/DemoAuthTest.cpp b/Engine/Source/Programs/Online/OnlineTestsCore/Source/Private/Tests/Auth/DemoAuthTest.cpp
--- a/Engine/Source/Programs/Online/OnlineTestsCore/Source/Private/Tests/Auth/DemoAuthTest.cpp
+++ b/Engine/Source/Programs/Online/OnlineTestsCore/Source/Private/Tests/Auth/DemoAuthTest.cpp
@@ -1,17 +1,13 @@
 // Copyright Epic Games, Inc. All Rights Reserved.
 
-#include "CoreMinimal.h"
-
-#include <catch2/catch_test_macros.hpp>
-
 #include "OnlineCatchHelper.h"
 
-
 #define AUTH_ERROR_TAG "[AUTH]"
 #define AUTH_ERROR_TEST_CASE(x, ...) ONLINE_TEST_CASE(x, AUTH_ERROR_TAG __VA_ARGS__)
 
 AUTH_ERROR_TEST_CASE("Basic login test", AUTH_ERROR_TAG)
 {
-	GetLoginPipeline(1);
+	FAccountId AccountId;
+	GetLoginPipeline({ AccountId });
 	RunToCompletion();
 }
diff --git a/Engine/Source/Programs/Online/OnlineTestsCore/Source/Private/Tests/Commerce/CheckoutTests.cpp b/Engine/Source/Programs/Online/OnlineTestsCore/Source/Private/Tests/Commerce/CheckoutTests.cpp
new file mode 100644
--- /dev/null
+++ b/Engine/Source/Programs/Online/OnlineTestsCore/Source/Private/Tests/Commerce/CheckoutTests.cpp
@@ -0,0 +1,120 @@
+// Copyright Epic Games, Inc. All Rights Reserved.
+
+#include "Helpers/Commerce/QueryOffersHelper.h"
+#include "Helpers/Commerce/GetOffersByIdHelper.h"
+#include "Helpers/Commerce/CheckoutHelper.h"
+#include "Helpers/Auth/AuthLogout.h"
+
+#define COMMERCE_TAG "[suite_commerce]"
+#define COMMERCE_CHECKOUT_TAG COMMERCE_TAG "[checkout]"
+
+#define COMMERCE_TEST_CASE(x, ...) ONLINE_TEST_CASE(x, COMMERCE_TAG __VA_ARGS__)
+
+COMMERCE_TEST_CASE("Verify that Checkout returns a fail message if the local user is not logged in", COMMERCE_CHECKOUT_TAG)
+{
+	FAccountId AccountId;
+	TSharedPtr<FPlatformUserId> AccountPlatformUserId = MakeShared<FPlatformUserId>();
+
+	FCommerceCheckout::Params OpCheckoutParams;
+	FCheckoutHelper::FHelperParams CheckoutHelperParams;
+	CheckoutHelperParams.OpParams = &OpCheckoutParams;
+	CheckoutHelperParams.ExpectedError = TOnlineResult<FCommerceCheckout>(Errors::NotLoggedIn());
+
+	FTestPipeline& LoginPipeline = GetLoginPipeline({ AccountId });
+	CheckoutHelperParams.OpParams->LocalAccountId = AccountId;
+
+	bool bLogout = false;
+
+	LoginPipeline
+		.EmplaceLambda([&AccountId, AccountPlatformUserId](SubsystemType OnlineSubsystem)
+			{
+				UE::Online::IAuthPtr OnlineAuthPtr = OnlineSubsystem->GetAuthInterface();
+				REQUIRE(OnlineAuthPtr);
+				UE::Online::TOnlineResult<UE::Online::FAuthGetLocalOnlineUserByOnlineAccountId> UserPlatfromUserIdResult = OnlineAuthPtr->GetLocalOnlineUserByOnlineAccountId({ AccountId });
+				REQUIRE(UserPlatfromUserIdResult.IsOk());
+				CHECK(UserPlatfromUserIdResult.TryGetOkValue() != nullptr);
+				*AccountPlatformUserId = UserPlatfromUserIdResult.TryGetOkValue()->AccountInfo->PlatformUserId;
+			})
+		.EmplaceStep<FAuthLogoutStep>(MoveTemp(AccountPlatformUserId))
+		.EmplaceStep<FCheckoutHelper>(MoveTemp(CheckoutHelperParams));
+
+	RunToCompletion(bLogout);
+}
+
+COMMERCE_TEST_CASE("Verify that Checkout returns a fail message of the given local user ID does not match the actual local user ID", COMMERCE_CHECKOUT_TAG)
+{
+	FAccountId AccountId;
+
+	FCommerceCheckout::Params OpCheckoutParams;
+	FCheckoutHelper::FHelperParams CheckoutHelperParams;
+	CheckoutHelperParams.OpParams = &OpCheckoutParams;
+	CheckoutHelperParams.ExpectedError = TOnlineResult<FCommerceCheckout>(Errors::NotLoggedIn());
+
+	FTestPipeline& LoginPipeline = GetLoginPipeline({ AccountId });
+	CheckoutHelperParams.OpParams->LocalAccountId = FAccountId();
+
+	LoginPipeline
+		.EmplaceStep<FCheckoutHelper>(MoveTemp(CheckoutHelperParams));
+
+	RunToCompletion();
+}
+
+COMMERCE_TEST_CASE("Verify that Checkout initiates the checkout process when given an Offers array with one offer", COMMERCE_CHECKOUT_TAG)
+{
+	SKIP("Test requires web browser to be enabled to run checkout process");
+
+	FAccountId AccountId;
+	TOptional<uint32_t> ExpectedOffersNum = 1;
+
+	FCommerceQueryOffers::Params OpQueryOffersParams;
+	FQueryOffersHelper::FHelperParams QueryOffersHelperParams;
+	QueryOffersHelperParams.OpParams = &OpQueryOffersParams;
+
+	FCommerceGetOffersById::Params OpGetOffersByIdParams;
+	FGetOffersByIdHelper::FHelperParams GetOffersByIdHelperParams;
+	GetOffersByIdHelperParams.OpParams = &OpGetOffersByIdParams;
+
+	FCommerceCheckout::Params OpCheckoutParams;
+	FCheckoutHelper::FHelperParams CheckoutHelperParams;
+	CheckoutHelperParams.OpParams = &OpCheckoutParams;
+
+	FTestPipeline& LoginPipeline = GetLoginPipeline({ AccountId });
+
+	QueryOffersHelperParams.OpParams->LocalAccountId = AccountId;
+
+	GetOffersByIdHelperParams.OpParams->LocalAccountId = AccountId;
+	FOfferId OfferId2;
+	GConfig->GetString(TEXT("EOSSDK.Platform.OnlineTests"), TEXT("OfferId2"), OfferId2, GEngineIni);
+	GetOffersByIdHelperParams.OpParams->OfferIds = { OfferId2 };
+
+	FPurchaseOffer Purchase = { OfferId2, 1};
+	CheckoutHelperParams.OpParams->LocalAccountId = AccountId;
+	CheckoutHelperParams.OpParams->Offers = { Purchase };
+
+	LoginPipeline
+		.EmplaceStep<FQueryOffersHelper>(MoveTemp(QueryOffersHelperParams))
+		.EmplaceStep<FGetOffersByIdHelper>(MoveTemp(GetOffersByIdHelperParams), ExpectedOffersNum)
+		.EmplaceStep<FCheckoutHelper>(MoveTemp(CheckoutHelperParams));
+
+	RunToCompletion();
+}
+
+COMMERCE_TEST_CASE("Verify that Checkout initiates the checkout process when given an Offers array with multiple offers")
+{
+	// TODO
+}
+
+COMMERCE_TEST_CASE("Verify that Checkout does not initiate the checkout process when given an empty Offers array")
+{
+	// TODO
+}
+
+COMMERCE_TEST_CASE("Verify that Checkout returns the correct TransactionId after a completed purchase")
+{
+	// TODO
+}
+
+COMMERCE_TEST_CASE("Verify that Checkout does not return a TransactionId if the purchase is incomplete")
+{
+	// TODO
+}
diff --git a/Engine/Source/Programs/Online/OnlineTestsCore/Source/Private/Tests/Commerce/GetEntitlementsTests.cpp b/Engine/Source/Programs/Online/OnlineTestsCore/Source/Private/Tests/Commerce/GetEntitlementsTests.cpp
new file mode 100644
--- /dev/null
+++ b/Engine/Source/Programs/Online/OnlineTestsCore/Source/Private/Tests/Commerce/GetEntitlementsTests.cpp
@@ -0,0 +1,66 @@
+// Copyright Epic Games, Inc. All Rights Reserved.
+
+#include "Helpers/Commerce/GetEntitlementsHelper.h"
+#include "Helpers/Auth/AuthLogout.h"
+
+#define COMMERCE_TAG "[suite_commerce]"
+#define COMMERCE_TEST_CASE(x, ...) ONLINE_TEST_CASE(x, COMMERCE_TAG __VA_ARGS__)
+
+COMMERCE_TEST_CASE("Verify that GetEntitlements returns a fail message if the local user is not logged in")
+{
+	FAccountId AccountId;
+	TSharedPtr<FPlatformUserId> AccountPlatformUserId = MakeShared<FPlatformUserId>();
+
+	FCommerceGetEntitlements::Params OpGetEntitlementsParams;
+	FGetEntitlementsHelper::FHelperParams GetEntitlementsHelperParams;
+	GetEntitlementsHelperParams.OpParams = &OpGetEntitlementsParams;
+	GetEntitlementsHelperParams.ExpectedError = TOnlineResult<FCommerceGetEntitlements>(Errors::NotLoggedIn());
+
+	FTestPipeline& LoginPipeline = GetLoginPipeline({ AccountId });
+	GetEntitlementsHelperParams.OpParams->LocalAccountId = AccountId;
+
+	bool bLogout = false;
+
+	LoginPipeline
+		.EmplaceLambda([&AccountId, AccountPlatformUserId](SubsystemType OnlineSubsystem)
+			{
+				UE::Online::IAuthPtr OnlineAuthPtr = OnlineSubsystem->GetAuthInterface();
+				REQUIRE(OnlineAuthPtr);
+				UE::Online::TOnlineResult<UE::Online::FAuthGetLocalOnlineUserByOnlineAccountId> UserPlatfromUserIdResult = OnlineAuthPtr->GetLocalOnlineUserByOnlineAccountId({ AccountId });
+				REQUIRE(UserPlatfromUserIdResult.IsOk());
+				CHECK(UserPlatfromUserIdResult.TryGetOkValue() != nullptr);
+				*AccountPlatformUserId = UserPlatfromUserIdResult.TryGetOkValue()->AccountInfo->PlatformUserId;
+			})
+		.EmplaceStep<FAuthLogoutStep>(MoveTemp(AccountPlatformUserId))
+		.EmplaceStep<FGetEntitlementsHelper>(MoveTemp(GetEntitlementsHelperParams));
+
+	RunToCompletion(bLogout);
+}
+
+COMMERCE_TEST_CASE("Verify that GetEntitlements returns a fail message of the given local user ID does not match the actual local user ID")
+{
+	FAccountId AccountId;
+
+	FCommerceGetEntitlements::Params OpGetEntitlementsParams;
+	FGetEntitlementsHelper::FHelperParams GetEntitlementsHelperParams;
+	GetEntitlementsHelperParams.OpParams = &OpGetEntitlementsParams;
+	GetEntitlementsHelperParams.ExpectedError = TOnlineResult<FCommerceGetEntitlements>(Errors::NotLoggedIn());
+
+	FTestPipeline& LoginPipeline = GetLoginPipeline({ AccountId });
+	GetEntitlementsHelperParams.OpParams->LocalAccountId = FAccountId();
+
+	LoginPipeline
+		.EmplaceStep<FGetEntitlementsHelper>(MoveTemp(GetEntitlementsHelperParams));
+
+	RunToCompletion();
+}
+
+COMMERCE_TEST_CASE("Verify that GetEntitlements returns a list of all cached entitlements")
+{
+	// TODO
+}
+
+COMMERCE_TEST_CASE("Verify that GetEntitlements returns an empty list when no entitlements are cached")
+{
+	// TODO
+}
\ No newline at end of file
diff --git a/Engine/Source/Programs/Online/OnlineTestsCore/Source/Private/Tests/Commerce/GetOffersByIdTests.cpp b/Engine/Source/Programs/Online/OnlineTestsCore/Source/Private/Tests/Commerce/GetOffersByIdTests.cpp
new file mode 100644
--- /dev/null
+++ b/Engine/Source/Programs/Online/OnlineTestsCore/Source/Private/Tests/Commerce/GetOffersByIdTests.cpp
@@ -0,0 +1,259 @@
+// Copyright Epic Games, Inc. All Rights Reserved.
+
+#include "Helpers/Commerce/QueryOffersHelper.h"
+#include "Helpers/Commerce/GetOffersByIdHelper.h"
+#include "Helpers/Auth/AuthLogout.h"
+
+#define COMMERCE_TAG "[suite_commerce]"
+#define COMMERCE_GETOFFERSBYID_TAG COMMERCE_TAG "[getoffersbyid]"
+
+#define COMMERCE_TEST_CASE(x, ...) ONLINE_TEST_CASE(x, COMMERCE_TAG __VA_ARGS__)
+
+COMMERCE_TEST_CASE("Verify that GetOffersById returns a fail message if the local user is not logged in", COMMERCE_GETOFFERSBYID_TAG)
+{
+	FAccountId AccountId;
+	TSharedPtr<FPlatformUserId> AccountPlatformUserId = MakeShared<FPlatformUserId>();
+
+	FCommerceGetOffersById::Params OpGetOffersByIdParams;
+	FGetOffersByIdHelper::FHelperParams GetOffersByIdHelperParams;
+	GetOffersByIdHelperParams.OpParams = &OpGetOffersByIdParams;
+	GetOffersByIdHelperParams.ExpectedError = TOnlineResult<FCommerceGetOffersById>(Errors::NotLoggedIn());
+
+	FTestPipeline& LoginPipeline = GetLoginPipeline({ AccountId });
+	GetOffersByIdHelperParams.OpParams->LocalAccountId = AccountId;
+
+	bool bLogout = false;
+
+	LoginPipeline
+		.EmplaceLambda([&AccountId, AccountPlatformUserId](SubsystemType OnlineSubsystem)
+			{
+				UE::Online::IAuthPtr OnlineAuthPtr = OnlineSubsystem->GetAuthInterface();
+				REQUIRE(OnlineAuthPtr);
+				UE::Online::TOnlineResult<UE::Online::FAuthGetLocalOnlineUserByOnlineAccountId> UserPlatfromUserIdResult = OnlineAuthPtr->GetLocalOnlineUserByOnlineAccountId({ AccountId });
+				REQUIRE(UserPlatfromUserIdResult.IsOk());
+				CHECK(UserPlatfromUserIdResult.TryGetOkValue() != nullptr);
+				*AccountPlatformUserId = UserPlatfromUserIdResult.TryGetOkValue()->AccountInfo->PlatformUserId;
+			})
+		.EmplaceStep<FAuthLogoutStep>(MoveTemp(AccountPlatformUserId))
+		.EmplaceStep<FGetOffersByIdHelper>(MoveTemp(GetOffersByIdHelperParams));
+
+	RunToCompletion(bLogout);
+}
+
+COMMERCE_TEST_CASE("Verify that GetOffersById returns a fail message of the given local user ID does not match the actual local user ID", COMMERCE_GETOFFERSBYID_TAG)
+{
+	FAccountId AccountId;
+
+	FCommerceGetOffersById::Params OpGetOffersByIdParams;
+	FGetOffersByIdHelper::FHelperParams GetOffersByIdHelperParams;
+	GetOffersByIdHelperParams.OpParams = &OpGetOffersByIdParams;
+	GetOffersByIdHelperParams.ExpectedError = TOnlineResult<FCommerceGetOffersById>(Errors::NotLoggedIn());
+
+	FTestPipeline& LoginPipeline = GetLoginPipeline({ AccountId });
+	GetOffersByIdHelperParams.OpParams->LocalAccountId = FAccountId();
+
+	LoginPipeline
+		.EmplaceStep<FGetOffersByIdHelper>(MoveTemp(GetOffersByIdHelperParams));
+
+	RunToCompletion();
+}
+
+COMMERCE_TEST_CASE("Verify that GetOffersById returns an empty list when given an empty list of IDs and there are no offers cached", COMMERCE_GETOFFERSBYID_TAG)
+{
+	FAccountId AccountId;
+	TOptional<uint32_t> ExpectedOffersNum = 0;
+
+	FCommerceQueryOffers::Params OpQueryOffersParams;
+	FQueryOffersHelper::FHelperParams QueryOffersHelperParams;
+	QueryOffersHelperParams.OpParams = &OpQueryOffersParams;
+
+	FCommerceGetOffersById::Params OpGetOffersByIdParams;
+	FGetOffersByIdHelper::FHelperParams GetOffersByIdHelperParams;
+	GetOffersByIdHelperParams.OpParams = &OpGetOffersByIdParams;
+
+	FTestPipeline& LoginPipeline = GetLoginPipeline({ AccountId });
+	QueryOffersHelperParams.OpParams->LocalAccountId = AccountId;
+	FString CatalogId;
+	GConfig->GetString(TEXT("EOSSDK.Platform.OnlineTests"), TEXT("OfferNamespace0Items"), CatalogId, GEngineIni);
+	QueryOffersHelperParams.OpParams->OverrideCatalogNamespace = CatalogId;
+	GetOffersByIdHelperParams.OpParams->LocalAccountId = AccountId;
+	GetOffersByIdHelperParams.OpParams->OfferIds = {};
+
+	LoginPipeline
+		.EmplaceStep<FQueryOffersHelper>(MoveTemp(QueryOffersHelperParams))
+		.EmplaceStep<FGetOffersByIdHelper>(MoveTemp(GetOffersByIdHelperParams), ExpectedOffersNum);
+
+	RunToCompletion();
+}
+
+COMMERCE_TEST_CASE("Verify that GetOffersById returns an empty list when given an empty list of IDs and there are multiple cached offers", COMMERCE_GETOFFERSBYID_TAG)
+{
+	FAccountId AccountId;
+	TOptional<uint32_t> ExpectedOffersNum = 0;
+
+	FCommerceQueryOffers::Params OpQueryOffersParams;
+	FQueryOffersHelper::FHelperParams QueryOffersHelperParams;
+	QueryOffersHelperParams.OpParams = &OpQueryOffersParams;
+
+	FCommerceGetOffersById::Params OpGetOffersByIdParams;
+	FGetOffersByIdHelper::FHelperParams GetOffersByIdHelperParams;
+	GetOffersByIdHelperParams.OpParams = &OpGetOffersByIdParams;
+
+	FTestPipeline& LoginPipeline = GetLoginPipeline({ AccountId });
+	QueryOffersHelperParams.OpParams->LocalAccountId = AccountId;
+	GetOffersByIdHelperParams.OpParams->LocalAccountId = AccountId;
+	GetOffersByIdHelperParams.OpParams->OfferIds = {};
+
+	LoginPipeline
+		.EmplaceStep<FQueryOffersHelper>(MoveTemp(QueryOffersHelperParams))
+		.EmplaceStep<FGetOffersByIdHelper>(MoveTemp(GetOffersByIdHelperParams), ExpectedOffersNum);
+
+	RunToCompletion();
+}
+
+COMMERCE_TEST_CASE("Verify that GetOffersById returns a fail message when given a populated list of IDs and there are no cached offers", COMMERCE_GETOFFERSBYID_TAG)
+{
+	FAccountId AccountId;
+	TOptional<uint32_t> ExpectedOffersNum = 0;
+
+	FCommerceQueryOffers::Params OpQueryOffersParams;
+	FQueryOffersHelper::FHelperParams QueryOffersHelperParams;
+	QueryOffersHelperParams.OpParams = &OpQueryOffersParams;
+
+	FCommerceGetOffersById::Params OpGetOffersByIdParams;
+	FGetOffersByIdHelper::FHelperParams GetOffersByIdHelperParams;
+	GetOffersByIdHelperParams.OpParams = &OpGetOffersByIdParams;
+	GetOffersByIdHelperParams.ExpectedError = TOnlineResult<FCommerceGetOffersById>(Errors::NotFound());
+
+	FTestPipeline& LoginPipeline = GetLoginPipeline({ AccountId });
+	QueryOffersHelperParams.OpParams->LocalAccountId = AccountId;
+	FString CatalogId;
+	GConfig->GetString(TEXT("EOSSDK.Platform.OnlineTests"), TEXT("OfferNamespace0Items"), CatalogId, GEngineIni);
+	QueryOffersHelperParams.OpParams->OverrideCatalogNamespace = CatalogId;
+	GetOffersByIdHelperParams.OpParams->LocalAccountId = AccountId;
+	FOfferId OfferId1NotExisting;
+	FOfferId OfferId2NotExisting;
+	GConfig->GetString(TEXT("EOSSDK.Platform.OnlineTests"), TEXT("OfferId1NotExisting"), OfferId1NotExisting, GEngineIni);
+	GConfig->GetString(TEXT("EOSSDK.Platform.OnlineTests"), TEXT("OfferId2NotExisting"), OfferId2NotExisting, GEngineIni);
+	GetOffersByIdHelperParams.OpParams->OfferIds = { OfferId1NotExisting, OfferId2NotExisting };
+
+	LoginPipeline
+		.EmplaceStep<FQueryOffersHelper>(MoveTemp(QueryOffersHelperParams))
+		.EmplaceStep<FGetOffersByIdHelper>(MoveTemp(GetOffersByIdHelperParams), ExpectedOffersNum);
+
+	RunToCompletion();
+}
+
+COMMERCE_TEST_CASE("Verify that GetOffersById returns an empty list when given a populated list of IDs and none of them match any cached offers", COMMERCE_GETOFFERSBYID_TAG)
+{
+	FAccountId AccountId;
+	TOptional<uint32_t> ExpectedOffersNum = 0;
+
+	FCommerceQueryOffers::Params OpQueryOffersParams;
+	FQueryOffersHelper::FHelperParams QueryOffersHelperParams;
+	QueryOffersHelperParams.OpParams = &OpQueryOffersParams;
+
+	FCommerceGetOffersById::Params OpGetOffersByIdParams;
+	FGetOffersByIdHelper::FHelperParams GetOffersByIdHelperParams;
+	GetOffersByIdHelperParams.OpParams = &OpGetOffersByIdParams;
+
+	FTestPipeline& LoginPipeline = GetLoginPipeline({ AccountId });
+	QueryOffersHelperParams.OpParams->LocalAccountId = AccountId;
+	GetOffersByIdHelperParams.OpParams->LocalAccountId = AccountId;
+	FOfferId OfferId1NotExisting;
+	FOfferId OfferId2NotExisting;
+	GConfig->GetString(TEXT("EOSSDK.Platform.OnlineTests"), TEXT("OfferId1NotExisting"), OfferId1NotExisting, GEngineIni);
+	GConfig->GetString(TEXT("EOSSDK.Platform.OnlineTests"), TEXT("OfferId2NotExisting"), OfferId2NotExisting, GEngineIni);
+	GetOffersByIdHelperParams.OpParams->OfferIds = { OfferId1NotExisting, OfferId2NotExisting };
+
+	LoginPipeline
+		.EmplaceStep<FQueryOffersHelper>(MoveTemp(QueryOffersHelperParams))
+		.EmplaceStep<FGetOffersByIdHelper>(MoveTemp(GetOffersByIdHelperParams), ExpectedOffersNum);
+
+	RunToCompletion();
+}
+
+COMMERCE_TEST_CASE("Verify that GetOffersById returns correct list of one offer when given a populated list of IDs where one ID is cached and another is not", COMMERCE_GETOFFERSBYID_TAG)
+{
+	FAccountId AccountId;
+	TOptional<uint32_t> ExpectedOffersNum = 1;
+
+	FCommerceQueryOffers::Params OpQueryOffersParams;
+	FQueryOffersHelper::FHelperParams QueryOffersHelperParams;
+	QueryOffersHelperParams.OpParams = &OpQueryOffersParams;
+
+	FCommerceGetOffersById::Params OpGetOffersByIdParams;
+	FGetOffersByIdHelper::FHelperParams GetOffersByIdHelperParams;
+	GetOffersByIdHelperParams.OpParams = &OpGetOffersByIdParams;
+
+	FTestPipeline& LoginPipeline = GetLoginPipeline({ AccountId });
+	QueryOffersHelperParams.OpParams->LocalAccountId = AccountId;
+	GetOffersByIdHelperParams.OpParams->LocalAccountId = AccountId;
+	FOfferId OfferId1NotExisting;
+	FOfferId OfferId2;
+	GConfig->GetString(TEXT("EOSSDK.Platform.OnlineTests"), TEXT("OfferId1NotExisting"), OfferId1NotExisting, GEngineIni);
+	GConfig->GetString(TEXT("EOSSDK.Platform.OnlineTests"), TEXT("OfferId2"), OfferId2, GEngineIni);
+	GetOffersByIdHelperParams.OpParams->OfferIds = { OfferId1NotExisting, OfferId2 };
+
+	LoginPipeline
+		.EmplaceStep<FQueryOffersHelper>(MoveTemp(QueryOffersHelperParams))
+		.EmplaceStep<FGetOffersByIdHelper>(MoveTemp(GetOffersByIdHelperParams), ExpectedOffersNum);
+
+	RunToCompletion();
+}
+
+COMMERCE_TEST_CASE("Verify that GetOffersById returns the correct list of one offer when given the ID for one cached offer", COMMERCE_GETOFFERSBYID_TAG)
+{
+	FAccountId AccountId;
+	TOptional<uint32_t> ExpectedOffersNum = 1;
+
+	FCommerceQueryOffers::Params OpQueryOffersParams;
+	FQueryOffersHelper::FHelperParams QueryOffersHelperParams;
+	QueryOffersHelperParams.OpParams = &OpQueryOffersParams;
+
+	FCommerceGetOffersById::Params OpGetOffersByIdParams;
+	FGetOffersByIdHelper::FHelperParams GetOffersByIdHelperParams;
+	GetOffersByIdHelperParams.OpParams = &OpGetOffersByIdParams;
+
+	FTestPipeline& LoginPipeline = GetLoginPipeline({ AccountId });
+	QueryOffersHelperParams.OpParams->LocalAccountId = AccountId;
+	GetOffersByIdHelperParams.OpParams->LocalAccountId = AccountId;
+	FOfferId OfferId1;
+	GConfig->GetString(TEXT("EOSSDK.Platform.OnlineTests"), TEXT("OfferId1"), OfferId1, GEngineIni);
+	GetOffersByIdHelperParams.OpParams->OfferIds = { OfferId1 };
+
+	LoginPipeline
+		.EmplaceStep<FQueryOffersHelper>(MoveTemp(QueryOffersHelperParams))
+		.EmplaceStep<FGetOffersByIdHelper>(MoveTemp(GetOffersByIdHelperParams), ExpectedOffersNum);
+
+	RunToCompletion();
+}
+
+COMMERCE_TEST_CASE("Verify that GetOffersById returns the correct list of offers when given a populated list of multiple cached IDs", COMMERCE_GETOFFERSBYID_TAG)
+{
+	FAccountId AccountId;
+	TOptional<uint32_t> ExpectedOffersNum = 2;
+
+	FCommerceQueryOffers::Params OpQueryOffersParams;
+	FQueryOffersHelper::FHelperParams QueryOffersHelperParams;
+	QueryOffersHelperParams.OpParams = &OpQueryOffersParams;
+
+	FCommerceGetOffersById::Params OpGetOffersByIdParams;
+	FGetOffersByIdHelper::FHelperParams GetOffersByIdHelperParams;
+	GetOffersByIdHelperParams.OpParams = &OpGetOffersByIdParams;
+
+	FTestPipeline& LoginPipeline = GetLoginPipeline({ AccountId });
+	QueryOffersHelperParams.OpParams->LocalAccountId = AccountId;
+	GetOffersByIdHelperParams.OpParams->LocalAccountId = AccountId;
+	FOfferId OfferId1;
+	FOfferId OfferId2;
+	GConfig->GetString(TEXT("EOSSDK.Platform.OnlineTests"), TEXT("OfferId1"), OfferId1, GEngineIni);
+	GConfig->GetString(TEXT("EOSSDK.Platform.OnlineTests"), TEXT("OfferId2"), OfferId2, GEngineIni);
+	GetOffersByIdHelperParams.OpParams->OfferIds = { OfferId1, OfferId2 };
+
+	LoginPipeline
+		.EmplaceStep<FQueryOffersHelper>(MoveTemp(QueryOffersHelperParams))
+		.EmplaceStep<FGetOffersByIdHelper>(MoveTemp(GetOffersByIdHelperParams), ExpectedOffersNum);
+
+	RunToCompletion();
+}
\ No newline at end of file
diff --git a/Engine/Source/Programs/Online/OnlineTestsCore/Source/Private/Tests/Commerce/GetOffersTests.cpp b/Engine/Source/Programs/Online/OnlineTestsCore/Source/Private/Tests/Commerce/GetOffersTests.cpp
new file mode 100644
--- /dev/null
+++ b/Engine/Source/Programs/Online/OnlineTestsCore/Source/Private/Tests/Commerce/GetOffersTests.cpp
@@ -0,0 +1,137 @@
+// Copyright Epic Games, Inc. All Rights Reserved.
+
+#include "Helpers/Commerce/QueryOffersHelper.h"
+#include "Helpers/Commerce/GetOffersHelper.h"
+#include "Helpers/Auth/AuthLogout.h"
+
+#define COMMERCE_TAG "[suite_commerce]"
+#define COMMERCE_GETOFFERS_TAG COMMERCE_TAG "[getoffers]"
+
+#define COMMERCE_TEST_CASE(x, ...) ONLINE_TEST_CASE(x, COMMERCE_TAG __VA_ARGS__)
+
+COMMERCE_TEST_CASE("Verify that GetOffers returns a fail message if the local user is not logged in", COMMERCE_GETOFFERS_TAG)
+{
+	FAccountId AccountId;
+	TSharedPtr<FPlatformUserId> AccountPlatformUserId = MakeShared<FPlatformUserId>();
+
+	FCommerceGetOffers::Params OpGetOffersParams;
+	FGetOffersHelper::FHelperParams GetOffersHelperParams;
+	GetOffersHelperParams.OpParams = &OpGetOffersParams;
+	GetOffersHelperParams.ExpectedError = TOnlineResult<FCommerceGetOffers>(Errors::NotLoggedIn());
+
+	FTestPipeline& LoginPipeline = GetLoginPipeline({ AccountId });
+	GetOffersHelperParams.OpParams->LocalAccountId = AccountId;
+
+	bool bLogout = false;
+
+	LoginPipeline
+		.EmplaceLambda([&AccountId, AccountPlatformUserId](SubsystemType OnlineSubsystem)
+			{
+				UE::Online::IAuthPtr OnlineAuthPtr = OnlineSubsystem->GetAuthInterface();
+				REQUIRE(OnlineAuthPtr);
+				UE::Online::TOnlineResult<UE::Online::FAuthGetLocalOnlineUserByOnlineAccountId> UserPlatfromUserIdResult = OnlineAuthPtr->GetLocalOnlineUserByOnlineAccountId({ AccountId });
+				REQUIRE(UserPlatfromUserIdResult.IsOk());
+				CHECK(UserPlatfromUserIdResult.TryGetOkValue() != nullptr);
+				*AccountPlatformUserId = UserPlatfromUserIdResult.TryGetOkValue()->AccountInfo->PlatformUserId;
+			})
+		.EmplaceStep<FAuthLogoutStep>(MoveTemp(AccountPlatformUserId))
+		.EmplaceStep<FGetOffersHelper>(MoveTemp(GetOffersHelperParams));
+
+	RunToCompletion(bLogout);
+}
+
+COMMERCE_TEST_CASE("Verify that GetOffers returns a fail message of the given local user ID does not match the actual local user ID", COMMERCE_GETOFFERS_TAG)
+{
+	FAccountId AccountId;
+
+	FCommerceGetOffers::Params OpGetOffersParams;
+	FGetOffersHelper::FHelperParams GetOffersHelperParams;
+	GetOffersHelperParams.OpParams = &OpGetOffersParams;
+	GetOffersHelperParams.ExpectedError = TOnlineResult<FCommerceGetOffers>(Errors::NotLoggedIn());
+
+	FTestPipeline& LoginPipeline = GetLoginPipeline({ AccountId });
+	GetOffersHelperParams.OpParams->LocalAccountId = FAccountId();
+
+	LoginPipeline
+		.EmplaceStep<FGetOffersHelper>(MoveTemp(GetOffersHelperParams));
+
+	RunToCompletion();
+}
+
+COMMERCE_TEST_CASE("Verify that GetOffers returns an empty list if there are no cached offers", COMMERCE_GETOFFERS_TAG)
+{
+	FAccountId AccountId;
+	TOptional<uint32_t> ExpectedOffersNum = 0;
+
+	FCommerceQueryOffers::Params OpQueryOffersParams;
+	FQueryOffersHelper::FHelperParams QueryOffersHelperParams;
+	QueryOffersHelperParams.OpParams = &OpQueryOffersParams;
+
+	FCommerceGetOffers::Params OpGetOffersParams;
+	FGetOffersHelper::FHelperParams GetOffersHelperParams;
+	GetOffersHelperParams.OpParams = &OpGetOffersParams;
+
+	FTestPipeline& LoginPipeline = GetLoginPipeline({ AccountId });
+	QueryOffersHelperParams.OpParams->LocalAccountId = AccountId;
+	FString CatalogId;
+	GConfig->GetString(TEXT("EOSSDK.Platform.OnlineTests"), TEXT("OfferNamespace0Items"), CatalogId, GEngineIni);
+	QueryOffersHelperParams.OpParams->OverrideCatalogNamespace = CatalogId;
+	GetOffersHelperParams.OpParams->LocalAccountId = AccountId;
+
+	LoginPipeline
+		.EmplaceStep<FQueryOffersHelper>(MoveTemp(QueryOffersHelperParams))
+		.EmplaceStep<FGetOffersHelper>(MoveTemp(GetOffersHelperParams), ExpectedOffersNum);
+
+	RunToCompletion();
+}
+
+COMMERCE_TEST_CASE("Verify that GetOffers returns a correct list of one offer if there is only one cached offer", COMMERCE_GETOFFERS_TAG)
+{
+	FAccountId AccountId;
+	TOptional<uint32_t> ExpectedOffersNum = 1;
+
+	FCommerceQueryOffers::Params OpQueryOffersParams;
+	FQueryOffersHelper::FHelperParams QueryOffersHelperParams;
+	QueryOffersHelperParams.OpParams = &OpQueryOffersParams;
+
+	FCommerceGetOffers::Params OpGetOffersParams;
+	FGetOffersHelper::FHelperParams GetOffersHelperParams;
+	GetOffersHelperParams.OpParams = &OpGetOffersParams;
+
+	FTestPipeline& LoginPipeline = GetLoginPipeline({ AccountId });
+	QueryOffersHelperParams.OpParams->LocalAccountId = AccountId;
+	FString CatalogId;
+	GConfig->GetString(TEXT("EOSSDK.Platform.OnlineTests"), TEXT("OfferNamespace1Items"), CatalogId, GEngineIni);
+	QueryOffersHelperParams.OpParams->OverrideCatalogNamespace = CatalogId;
+	GetOffersHelperParams.OpParams->LocalAccountId = AccountId;
+
+	LoginPipeline
+		.EmplaceStep<FQueryOffersHelper>(MoveTemp(QueryOffersHelperParams))
+		.EmplaceStep<FGetOffersHelper>(MoveTemp(GetOffersHelperParams), ExpectedOffersNum);
+
+	RunToCompletion();
+}
+
+COMMERCE_TEST_CASE("Verify that GetOffers returns the correct list if there are cached offers", COMMERCE_GETOFFERS_TAG)
+{
+	FAccountId AccountId;
+	TOptional<uint32_t> ExpectedOffersNum = 4;
+
+	FCommerceQueryOffers::Params OpQueryOffersParams;
+	FQueryOffersHelper::FHelperParams QueryOffersHelperParams;
+	QueryOffersHelperParams.OpParams = &OpQueryOffersParams;
+
+	FCommerceGetOffers::Params OpGetOffersParams;
+	FGetOffersHelper::FHelperParams GetOffersHelperParams;
+	GetOffersHelperParams.OpParams = &OpGetOffersParams;
+
+	FTestPipeline& LoginPipeline = GetLoginPipeline({ AccountId });
+	QueryOffersHelperParams.OpParams->LocalAccountId = AccountId;
+	GetOffersHelperParams.OpParams->LocalAccountId = AccountId;
+
+	LoginPipeline
+		.EmplaceStep<FQueryOffersHelper>(MoveTemp(QueryOffersHelperParams))
+		.EmplaceStep<FGetOffersHelper>(MoveTemp(GetOffersHelperParams), ExpectedOffersNum);
+
+	RunToCompletion();
+}
\ No newline at end of file
diff --git a/Engine/Source/Programs/Online/OnlineTestsCore/Source/Private/Tests/Commerce/OnPurchaseCompletedTests.cpp b/Engine/Source/Programs/Online/OnlineTestsCore/Source/Private/Tests/Commerce/OnPurchaseCompletedTests.cpp
new file mode 100644
--- /dev/null
+++ b/Engine/Source/Programs/Online/OnlineTestsCore/Source/Private/Tests/Commerce/OnPurchaseCompletedTests.cpp
@@ -0,0 +1,22 @@
+// Copyright Epic Games, Inc. All Rights Reserved.
+
+#include "OnlineCatchHelper.h"
+
+#define COMMERCE_TAG "[suite_commerce]"
+
+#define COMMERCE_TEST_CASE(x, ...) ONLINE_TEST_CASE(x, COMMERCE_TAG __VA_ARGS__)
+
+COMMERCE_TEST_CASE("Verify that OnPurchaseCompleted returns a fail message if the local user is not logged in")
+{
+	// TODO
+}
+
+COMMERCE_TEST_CASE("Verify that OnPurchaseCompleted returns a fail message of the given local user ID does not match the actual local user ID")
+{
+	// TODO
+}
+
+COMMERCE_TEST_CASE("Verify that OnPurchaseCompleted completes without an error")
+{
+	// TODO
+}
\ No newline at end of file
diff --git a/Engine/Source/Programs/Online/OnlineTestsCore/Source/Private/Tests/Commerce/QueryEntitlementsTests.cpp b/Engine/Source/Programs/Online/OnlineTestsCore/Source/Private/Tests/Commerce/QueryEntitlementsTests.cpp
new file mode 100644
--- /dev/null
+++ b/Engine/Source/Programs/Online/OnlineTestsCore/Source/Private/Tests/Commerce/QueryEntitlementsTests.cpp
@@ -0,0 +1,83 @@
+// Copyright Epic Games, Inc. All Rights Reserved.
+
+#include "Helpers/Commerce/GetEntitlementsHelper.h"
+#include "Helpers/Commerce/QueryEntitlementsHelper.h"
+#include "Helpers/Auth/AuthLogout.h"
+
+#define COMMERCE_TAG "[suite_commerce]"
+
+#define COMMERCE_TEST_CASE(x, ...) ONLINE_TEST_CASE(x, COMMERCE_TAG __VA_ARGS__)
+
+COMMERCE_TEST_CASE("Verify that QueryEntitlements returns a fail message if the local user is not logged in")
+{
+	FAccountId AccountId;
+	TSharedPtr<FPlatformUserId> AccountPlatformUserId = MakeShared<FPlatformUserId>();
+
+	FCommerceQueryEntitlements::Params OpQueryEntitlementsParams;
+	FQueryEntitlementsHelper::FHelperParams QueryEntitlementsHelperParams;
+	QueryEntitlementsHelperParams.ExpectedError = TOnlineResult<FCommerceQueryEntitlements>(Errors::NotLoggedIn());
+	QueryEntitlementsHelperParams.OpParams = &OpQueryEntitlementsParams;
+
+	FTestPipeline& LoginPipeline = GetLoginPipeline({ AccountId });
+	QueryEntitlementsHelperParams.OpParams->LocalAccountId = AccountId;
+
+	bool bLogout = false;
+
+	LoginPipeline
+		.EmplaceLambda([&AccountId, AccountPlatformUserId](SubsystemType OnlineSubsystem)
+			{
+				UE::Online::IAuthPtr OnlineAuthPtr = OnlineSubsystem->GetAuthInterface();
+				REQUIRE(OnlineAuthPtr);
+				UE::Online::TOnlineResult<UE::Online::FAuthGetLocalOnlineUserByOnlineAccountId> UserPlatfromUserIdResult = OnlineAuthPtr->GetLocalOnlineUserByOnlineAccountId({ AccountId });
+				REQUIRE(UserPlatfromUserIdResult.IsOk());
+				CHECK(UserPlatfromUserIdResult.TryGetOkValue() != nullptr);
+				*AccountPlatformUserId = UserPlatfromUserIdResult.TryGetOkValue()->AccountInfo->PlatformUserId;
+			})
+		.EmplaceStep<FAuthLogoutStep>(MoveTemp(AccountPlatformUserId))
+		.EmplaceStep<FQueryEntitlementsHelper>(MoveTemp(QueryEntitlementsHelperParams));
+
+			RunToCompletion(bLogout);
+}
+
+COMMERCE_TEST_CASE("Verify that QueryEntitlements returns a fail message of the given local user ID does not match the actual local user ID")
+{
+	FAccountId AccountId;
+
+	FCommerceQueryEntitlements::Params OpQueryEntitlementsParams;
+	FQueryEntitlementsHelper::FHelperParams QueryEntitlementsHelperParams;
+	QueryEntitlementsHelperParams.OpParams = &OpQueryEntitlementsParams;
+	QueryEntitlementsHelperParams.ExpectedError = TOnlineResult<FCommerceQueryEntitlements>(Errors::NotLoggedIn());
+
+	FTestPipeline& LoginPipeline = GetLoginPipeline({ AccountId });
+	QueryEntitlementsHelperParams.OpParams->LocalAccountId = FAccountId();
+
+	LoginPipeline
+		.EmplaceStep<FQueryEntitlementsHelper>(MoveTemp(QueryEntitlementsHelperParams));
+
+	RunToCompletion();
+}
+
+COMMERCE_TEST_CASE("Verify that QueryEntitlements caches all entitlements")
+{
+	FAccountId AccountId;
+	TOptional<uint32_t> ExpectedEntitlementsNum = 0;
+
+	FCommerceQueryEntitlements::Params OpQueryEntitlementsParams;
+	FQueryEntitlementsHelper::FHelperParams QueryEntitlementsHelperParams;
+	QueryEntitlementsHelperParams.OpParams = &OpQueryEntitlementsParams;
+
+	FCommerceGetEntitlements::Params OpGetEntitlementsParams;
+	FGetEntitlementsHelper::FHelperParams GetEntitlementsHelperParams;
+	GetEntitlementsHelperParams.OpParams = &OpGetEntitlementsParams;
+
+	FTestPipeline& LoginPipeline = GetLoginPipeline({ AccountId });
+	QueryEntitlementsHelperParams.OpParams->LocalAccountId = AccountId;
+	QueryEntitlementsHelperParams.OpParams->bIncludeRedeemed = true;
+	GetEntitlementsHelperParams.OpParams->LocalAccountId = AccountId;
+
+	LoginPipeline
+		.EmplaceStep<FQueryEntitlementsHelper>(MoveTemp(QueryEntitlementsHelperParams))
+		.EmplaceStep<FGetEntitlementsHelper>(MoveTemp(GetEntitlementsHelperParams), ExpectedEntitlementsNum);
+
+	RunToCompletion();
+}
\ No newline at end of file
diff --git a/Engine/Source/Programs/Online/OnlineTestsCore/Source/Private/Tests/Commerce/QueryOffersByIdTests.cpp b/Engine/Source/Programs/Online/OnlineTestsCore/Source/Private/Tests/Commerce/QueryOffersByIdTests.cpp
new file mode 100644
--- /dev/null
+++ b/Engine/Source/Programs/Online/OnlineTestsCore/Source/Private/Tests/Commerce/QueryOffersByIdTests.cpp
@@ -0,0 +1,94 @@
+// Copyright Epic Games, Inc. All Rights Reserved.
+
+#include "Helpers/Commerce/QueryOffersByIdHelper.h"
+#include "Helpers/Auth/AuthLogout.h"
+
+#define COMMERCE_TAG "[suite_commerce]"
+#define COMMERCE_QUERYOFFERSBYID_TAG COMMERCE_TAG "[queryoffersbyid]"
+#define COMMERCE_DISABLED_TAG COMMERCE_TAG "[commercedisabled]"
+
+#define COMMERCE_TEST_CASE(x, ...) ONLINE_TEST_CASE(x, COMMERCE_TAG __VA_ARGS__)
+
+COMMERCE_TEST_CASE("Verify that QueryOffersById returns a fail message if the local user is not logged in", COMMERCE_QUERYOFFERSBYID_TAG)
+{
+	FAccountId AccountId;
+	TSharedPtr<FPlatformUserId> AccountPlatformUserId = MakeShared<FPlatformUserId>();
+
+	FCommerceQueryOffersById::Params OpQueryOffersByIdParams;
+	FQueryOffersByIdHelper::FHelperParams QueryOffersByIdHelperParams;
+	QueryOffersByIdHelperParams.OpParams = &OpQueryOffersByIdParams;
+	QueryOffersByIdHelperParams.ExpectedError = TOnlineResult<FCommerceQueryOffersById>(Errors::NotLoggedIn());
+
+	FTestPipeline& LoginPipeline = GetLoginPipeline({ AccountId });
+	QueryOffersByIdHelperParams.OpParams->LocalAccountId = AccountId;
+
+	bool bLogout = false;
+
+	LoginPipeline
+		.EmplaceLambda([&AccountId, AccountPlatformUserId](SubsystemType OnlineSubsystem)
+			{
+				UE::Online::IAuthPtr OnlineAuthPtr = OnlineSubsystem->GetAuthInterface();
+				REQUIRE(OnlineAuthPtr);
+				UE::Online::TOnlineResult<UE::Online::FAuthGetLocalOnlineUserByOnlineAccountId> UserPlatfromUserIdResult = OnlineAuthPtr->GetLocalOnlineUserByOnlineAccountId({ AccountId });
+				REQUIRE(UserPlatfromUserIdResult.IsOk());
+				CHECK(UserPlatfromUserIdResult.TryGetOkValue() != nullptr);
+				*AccountPlatformUserId = UserPlatfromUserIdResult.TryGetOkValue()->AccountInfo->PlatformUserId;
+			})
+		.EmplaceStep<FAuthLogoutStep>(MoveTemp(AccountPlatformUserId))
+		.EmplaceStep<FQueryOffersByIdHelper>(MoveTemp(QueryOffersByIdHelperParams));
+
+	RunToCompletion(bLogout);
+}
+
+COMMERCE_TEST_CASE("Verify that QueryOffersById returns a fail message of the given local user ID does not match the actual local user ID", COMMERCE_QUERYOFFERSBYID_TAG)
+{
+	FAccountId AccountId;
+
+	FCommerceQueryOffersById::Params OpQueryOffersByIdParams;
+	FQueryOffersByIdHelper::FHelperParams QueryOffersByIdHelperParams;
+	QueryOffersByIdHelperParams.OpParams = &OpQueryOffersByIdParams;
+	QueryOffersByIdHelperParams.ExpectedError = TOnlineResult<FCommerceQueryOffersById>(Errors::NotLoggedIn());
+
+	FTestPipeline& LoginPipeline = GetLoginPipeline({ AccountId });
+	QueryOffersByIdHelperParams.OpParams->LocalAccountId = FAccountId();
+
+	LoginPipeline
+		.EmplaceStep<FQueryOffersByIdHelper>(MoveTemp(QueryOffersByIdHelperParams));
+
+	RunToCompletion();
+}
+
+COMMERCE_TEST_CASE("Verify that QueryOffersById returns an empty list when given an empty list of IDs and there are no existing offers", COMMERCE_DISABLED_TAG)
+{
+	// EOS QueryOffers interface doesn't support anything ID-specific, skipping the test
+}
+
+COMMERCE_TEST_CASE("Verify that QueryOffersById returns an empty list when given an empty list of IDs and there are multiple existing offers", COMMERCE_DISABLED_TAG)
+{
+	// EOS QueryOffers interface doesn't support anything ID-specific, skipping the test
+}
+
+COMMERCE_TEST_CASE("Verify that QueryOffersById returns a fail message when given a populated list of IDs and there are no existing offers", COMMERCE_DISABLED_TAG)
+{
+	// EOS QueryOffers interface doesn't support anything ID-specific, skipping the test
+}
+
+COMMERCE_TEST_CASE("Verify that QueryOffersById returns a fail message when given a populated list of IDs and none of them match existing offers", COMMERCE_DISABLED_TAG)
+{
+	// EOS QueryOffers interface doesn't support anything ID-specific, skipping the test
+}
+
+COMMERCE_TEST_CASE("Verify that QueryOffersById returns a fail message when given a populated list of IDs where one ID exists and another does not", COMMERCE_DISABLED_TAG)
+{
+	// EOS QueryOffers interface doesn't support anything ID-specific, skipping the test
+}
+
+COMMERCE_TEST_CASE("Verify that QueryOffersById returns the correct list of one offer when given the ID for one existing offer", COMMERCE_DISABLED_TAG)
+{
+	// EOS QueryOffers interface doesn't support anything ID-specific, skipping the test
+}
+
+COMMERCE_TEST_CASE("Verify that QueryOffersById returns the correct list of offers when given a populated list of multiple existing IDs", COMMERCE_DISABLED_TAG)
+{
+	// EOS QueryOffers interface doesn't support anything ID-specific, skipping the test
+}
\ No newline at end of file
diff --git a/Engine/Source/Programs/Online/OnlineTestsCore/Source/Private/Tests/Commerce/QueryOffersTests.cpp b/Engine/Source/Programs/Online/OnlineTestsCore/Source/Private/Tests/Commerce/QueryOffersTests.cpp
new file mode 100644
--- /dev/null
+++ b/Engine/Source/Programs/Online/OnlineTestsCore/Source/Private/Tests/Commerce/QueryOffersTests.cpp
@@ -0,0 +1,143 @@
+// Copyright Epic Games, Inc. All Rights Reserved.
+
+#include "Helpers/Commerce/GetOffersHelper.h"
+#include "Helpers/Commerce/QueryOffersHelper.h"
+#include "Helpers/Auth/AuthLogout.h"
+
+#define COMMERCE_TAG "[suite_commerce]"
+#define COMMERCE_QUERYOFFERS_TAG COMMERCE_TAG "[queryoffers]"
+
+#define COMMERCE_TEST_CASE(x, ...) ONLINE_TEST_CASE(x, COMMERCE_TAG __VA_ARGS__)
+
+COMMERCE_TEST_CASE("Basic compile test")
+{
+	TSharedPtr<ICommerce> Commerce;
+	CHECK(!Commerce.IsValid());
+}
+
+COMMERCE_TEST_CASE("Verify that QueryOffers returns a fail message if the local user is not logged in", COMMERCE_QUERYOFFERS_TAG)
+{
+	FAccountId AccountId;
+	TSharedPtr<FPlatformUserId> AccountPlatformUserId = MakeShared<FPlatformUserId>();
+
+	FCommerceQueryOffers::Params OpQueryOffersParams;
+	FQueryOffersHelper::FHelperParams QueryOffersHelperParams;
+	QueryOffersHelperParams.OpParams = &OpQueryOffersParams;
+	QueryOffersHelperParams.ExpectedError = TOnlineResult<FCommerceQueryOffers>(Errors::NotLoggedIn());
+
+	FTestPipeline& LoginPipeline = GetLoginPipeline({ AccountId });
+	QueryOffersHelperParams.OpParams->LocalAccountId = AccountId;
+
+	bool bLogout = false;
+
+	LoginPipeline
+		.EmplaceLambda([&AccountId, AccountPlatformUserId](SubsystemType OnlineSubsystem)
+			{
+				UE::Online::IAuthPtr OnlineAuthPtr = OnlineSubsystem->GetAuthInterface();
+				REQUIRE(OnlineAuthPtr);
+				UE::Online::TOnlineResult<UE::Online::FAuthGetLocalOnlineUserByOnlineAccountId> UserPlatfromUserIdResult = OnlineAuthPtr->GetLocalOnlineUserByOnlineAccountId({ AccountId });
+				REQUIRE(UserPlatfromUserIdResult.IsOk());
+				CHECK(UserPlatfromUserIdResult.TryGetOkValue() != nullptr);
+				*AccountPlatformUserId = UserPlatfromUserIdResult.TryGetOkValue()->AccountInfo->PlatformUserId;
+			})
+		.EmplaceStep<FAuthLogoutStep>(MoveTemp(AccountPlatformUserId))
+		.EmplaceStep<FQueryOffersHelper>(MoveTemp(QueryOffersHelperParams));
+
+	RunToCompletion(bLogout);
+}
+
+COMMERCE_TEST_CASE("Verify that QueryOffers returns a fail message of the given local user ID does not match the actual local user ID", COMMERCE_QUERYOFFERS_TAG)
+{
+	FAccountId AccountId;
+
+	FCommerceQueryOffers::Params OpQueryOffersParams;
+	FQueryOffersHelper::FHelperParams QueryOffersHelperParams;
+	QueryOffersHelperParams.OpParams = &OpQueryOffersParams;
+	QueryOffersHelperParams.ExpectedError = TOnlineResult<FCommerceQueryOffers>(Errors::NotLoggedIn());
+
+	FTestPipeline& LoginPipeline = GetLoginPipeline({ AccountId });
+	QueryOffersHelperParams.OpParams->LocalAccountId = FAccountId();
+
+	LoginPipeline
+		.EmplaceStep<FQueryOffersHelper>(MoveTemp(QueryOffersHelperParams));
+
+	RunToCompletion();
+}
+
+COMMERCE_TEST_CASE("Verify that QueryOffers caches an empty list if there are no offers", COMMERCE_QUERYOFFERS_TAG)
+{
+	FAccountId AccountId;
+	TOptional<uint32_t> ExpectedOffersNum = 0;
+
+	FCommerceQueryOffers::Params OpQueryOffersParams;
+	FQueryOffersHelper::FHelperParams QueryOffersHelperParams;
+	QueryOffersHelperParams.OpParams = &OpQueryOffersParams;
+
+	FCommerceGetOffers::Params OpGetOffersParams;
+	FGetOffersHelper::FHelperParams GetOffersHelperParams;
+	GetOffersHelperParams.OpParams = &OpGetOffersParams;
+
+	FTestPipeline& LoginPipeline = GetLoginPipeline({ AccountId });
+	QueryOffersHelperParams.OpParams->LocalAccountId = AccountId;
+	FString CatalogId;
+	GConfig->GetString(TEXT("EOSSDK.Platform.OnlineTests"), TEXT("OfferNamespace0Items"), CatalogId, GEngineIni);
+	QueryOffersHelperParams.OpParams->OverrideCatalogNamespace = CatalogId;
+	GetOffersHelperParams.OpParams->LocalAccountId = AccountId;
+
+	LoginPipeline
+		.EmplaceStep<FQueryOffersHelper>(MoveTemp(QueryOffersHelperParams))
+		.EmplaceStep<FGetOffersHelper>(MoveTemp(GetOffersHelperParams), ExpectedOffersNum);
+
+	RunToCompletion();
+}
+
+COMMERCE_TEST_CASE("Verify that QueryOffers caches a list of one offer if there is only one existing offer", COMMERCE_QUERYOFFERS_TAG)
+{
+	FAccountId AccountId;
+	TOptional<uint32_t> ExpectedOffersNum = 1;
+
+	FCommerceQueryOffers::Params OpQueryOffersParams;
+	FQueryOffersHelper::FHelperParams QueryOffersHelperParams;
+	QueryOffersHelperParams.OpParams = &OpQueryOffersParams;
+
+	FCommerceGetOffers::Params OpGetOffersParams;
+	FGetOffersHelper::FHelperParams GetOffersHelperParams;
+	GetOffersHelperParams.OpParams = &OpGetOffersParams;
+
+	FTestPipeline& LoginPipeline = GetLoginPipeline({ AccountId });
+	QueryOffersHelperParams.OpParams->LocalAccountId = AccountId;
+	FString CatalogId;
+	GConfig->GetString(TEXT("EOSSDK.Platform.OnlineTests"), TEXT("OfferNamespace1Items"), CatalogId, GEngineIni);
+	QueryOffersHelperParams.OpParams->OverrideCatalogNamespace = CatalogId;
+	GetOffersHelperParams.OpParams->LocalAccountId = AccountId;
+
+	LoginPipeline
+		.EmplaceStep<FQueryOffersHelper>(MoveTemp(QueryOffersHelperParams))
+		.EmplaceStep<FGetOffersHelper>(MoveTemp(GetOffersHelperParams), ExpectedOffersNum);
+
+	RunToCompletion();
+}
+
+COMMERCE_TEST_CASE("Verify that QueryOffers caches the list of all offers if there are multiple existing offers", COMMERCE_QUERYOFFERS_TAG)
+{
+	FAccountId AccountId;
+	TOptional<uint32_t> ExpectedOffersNum = 4;
+
+	FCommerceQueryOffers::Params OpQueryOffersParams;
+	FQueryOffersHelper::FHelperParams QueryOffersHelperParams;
+	QueryOffersHelperParams.OpParams = &OpQueryOffersParams;
+
+	FCommerceGetOffers::Params OpGetOffersParams;
+	FGetOffersHelper::FHelperParams GetOffersHelperParams;
+	GetOffersHelperParams.OpParams = &OpGetOffersParams;
+
+	FTestPipeline& LoginPipeline = GetLoginPipeline({ AccountId });
+	QueryOffersHelperParams.OpParams->LocalAccountId = AccountId;
+	GetOffersHelperParams.OpParams->LocalAccountId = AccountId;
+
+	LoginPipeline
+		.EmplaceStep<FQueryOffersHelper>(MoveTemp(QueryOffersHelperParams))
+		.EmplaceStep<FGetOffersHelper>(MoveTemp(GetOffersHelperParams), ExpectedOffersNum);
+
+	RunToCompletion();
+}
diff --git a/Engine/Source/Programs/Online/OnlineTestsCore/Source/Private/Tests/Commerce/QueryTransactionEntitlementsTests.cpp b/Engine/Source/Programs/Online/OnlineTestsCore/Source/Private/Tests/Commerce/QueryTransactionEntitlementsTests.cpp
new file mode 100644
--- /dev/null
+++ b/Engine/Source/Programs/Online/OnlineTestsCore/Source/Private/Tests/Commerce/QueryTransactionEntitlementsTests.cpp
@@ -0,0 +1,27 @@
+// Copyright Epic Games, Inc. All Rights Reserved.
+
+#include "OnlineCatchHelper.h"
+
+#define COMMERCE_TAG "[suite_commerce]"
+
+#define COMMERCE_TEST_CASE(x, ...) ONLINE_TEST_CASE(x, COMMERCE_TAG __VA_ARGS__)
+
+COMMERCE_TEST_CASE("Verify that QueryTransactionEntitlements returns a fail message if the local user is not logged in")
+{
+	// TODO
+}
+
+COMMERCE_TEST_CASE("Verify that QueryTransactionEntitlements returns a fail message of the given local user ID does not match the actual local user ID")
+{
+	// TODO
+}
+
+COMMERCE_TEST_CASE("Verify that QueryTransactionEntitlements returns the correct list of entitlements when given an existing TransactionId and caches them")
+{
+	// TODO
+}
+
+COMMERCE_TEST_CASE("Verify that QueryTransactionEntitlements returns a fail message when given a nonexisting TransactionId")
+{
+	// TODO
+}
diff --git a/Engine/Source/Programs/Online/OnlineTestsCore/Source/Private/Tests/Commerce/RedeemEntitlementTests.cpp b/Engine/Source/Programs/Online/OnlineTestsCore/Source/Private/Tests/Commerce/RedeemEntitlementTests.cpp
new file mode 100644
--- /dev/null
+++ b/Engine/Source/Programs/Online/OnlineTestsCore/Source/Private/Tests/Commerce/RedeemEntitlementTests.cpp
@@ -0,0 +1,57 @@
+// Copyright Epic Games, Inc. All Rights Reserved.
+
+#include "OnlineCatchHelper.h"
+
+#define COMMERCE_TAG "[suite_commerce]"
+
+#define COMMERCE_TEST_CASE(x, ...) ONLINE_TEST_CASE(x, COMMERCE_TAG __VA_ARGS__)
+
+COMMERCE_TEST_CASE("Verify that RedeemEntitlement returns a fail message if the local user is not logged in")
+{
+	// TODO
+}
+
+COMMERCE_TEST_CASE("Verify that RedeemEntitlement returns a fail message of the given local user ID does not match the actual local user ID")
+{
+	// TODO
+}
+
+COMMERCE_TEST_CASE("Verify that RedeemEntitlement redeems no entitlements when given an existing EntitlementId and Quantity 0")
+{
+	// TODO
+}
+
+COMMERCE_TEST_CASE("Verify that RedeemEntitlement redeems 1 of the correct entitlements when given an existing EntitlementId and Quantity 1")
+{
+	// TODO
+}
+
+COMMERCE_TEST_CASE("Verify that RedeemEntitlement redeems 2 of the correct entitlements when given an existing EntitlementId and Quantity 2")
+{
+	// TODO
+}
+
+COMMERCE_TEST_CASE("Verify that RedeemEntitlement redeems all of the given entitlements when given a Quantity equal to the available number for that entitlement")
+{
+	// TODO
+}
+
+COMMERCE_TEST_CASE("Verify that RedeemEntitlement returns a fail message when given a Quantity that is 1 more than the available number for the given entitlemenet")
+{
+	// TODO
+}
+
+COMMERCE_TEST_CASE("Verify that RedeemEntitlement returns a fail message when given a Quantity that is 2 more than the available number for the given entitlemenet")
+{
+	// TODO
+}
+
+COMMERCE_TEST_CASE("Verify that RedeemEntitlement returns a fail message when given a nonexisting EntitlementId")
+{
+	// TODO
+}
+
+COMMERCE_TEST_CASE("Verify that RedeemEntitlement returns all of the given entitlements when given a Quantity matching the number of entitlements available")
+{
+	// TODO
+}
\ No newline at end of file
diff --git a/Engine/Source/Programs/Online/OnlineTestsCore/Source/Private/Tests/Commerce/RetrieveS2STokenTests.cpp b/Engine/Source/Programs/Online/OnlineTestsCore/Source/Private/Tests/Commerce/RetrieveS2STokenTests.cpp
new file mode 100644
--- /dev/null
+++ b/Engine/Source/Programs/Online/OnlineTestsCore/Source/Private/Tests/Commerce/RetrieveS2STokenTests.cpp
@@ -0,0 +1,27 @@
+// Copyright Epic Games, Inc. All Rights Reserved.
+
+#include "OnlineCatchHelper.h"
+
+#define COMMERCE_TAG "[suite_commerce]"
+
+#define COMMERCE_TEST_CASE(x, ...) ONLINE_TEST_CASE(x, COMMERCE_TAG __VA_ARGS__)
+
+COMMERCE_TEST_CASE("Verify that RetrieveS2SToken returns a fail message if the local user is not logged in")
+{
+	// TODO
+}
+
+COMMERCE_TEST_CASE("Verify that RetrieveS2SToken returns a fail message of the given local user ID does not match the actual local user ID")
+{
+	// TODO
+}
+
+COMMERCE_TEST_CASE("Verify that RetrieveS2SToken successfully returns the correct Token for the given TokenType")
+{
+	// TODO
+}
+
+COMMERCE_TEST_CASE("Verify that RetrieveS2SToken returns a fail message if the given TokenType does not exist")
+{
+	// TODO
+}
diff --git a/Engine/Source/Programs/Online/OnlineTestsCore/Source/Private/Tests/Commerce/ShowStoreUITests.cpp b/Engine/Source/Programs/Online/OnlineTestsCore/Source/Private/Tests/Commerce/ShowStoreUITests.cpp
new file mode 100644
--- /dev/null
+++ b/Engine/Source/Programs/Online/OnlineTestsCore/Source/Private/Tests/Commerce/ShowStoreUITests.cpp
@@ -0,0 +1,26 @@
+// Copyright Epic Games, Inc. All Rights Reserved.
+
+#include "OnlineCatchHelper.h"
+
+#define COMMERCE_TAG "[suite_commerce]"
+
+#define COMMERCE_DISABLED_TAG COMMERCE_TAG "[commercedisabled]"
+#define COMMERCE_TEST_CASE(x, ...) ONLINE_TEST_CASE(x, COMMERCE_TAG __VA_ARGS__)
+
+COMMERCE_TEST_CASE("Verify that ShowStoreUI returns a fail message if the local user is not logged in", COMMERCE_DISABLED_TAG)
+{
+	// TODO
+	// The ShowStoreUI EOS interface is not implemented
+}
+
+COMMERCE_TEST_CASE("Verify that ShowStoreUI returns a fail message of the given local user ID does not match the actual local user ID", COMMERCE_DISABLED_TAG)
+{
+	// TODO
+	// The ShowStoreUI EOS interface is not implemented
+}
+
+COMMERCE_TEST_CASE("Verify that ShowStoreUI displays platform store UI", COMMERCE_DISABLED_TAG) // May not be possible to automate currently
+{
+	// TODO
+	// The ShowStoreUI EOS interface is not implemented
+}
\ No newline at end of file
diff --git a/Engine/Source/Programs/Online/OnlineTestsCore/Source/Private/Tests/Commerce/TestCommerce.cpp b/Engine/Source/Programs/Online/OnlineTestsCore/Source/Private/Tests/Commerce/TestCommerce.cpp
deleted file mode 100644
diff --git a/Engine/Source/Programs/Online/OnlineTestsCore/Source/Private/Tests/DelegateAdapter/DelegateAdapterTest.cpp b/Engine/Source/Programs/Online/OnlineTestsCore/Source/Private/Tests/DelegateAdapter/DelegateAdapterTest.cpp
--- a/Engine/Source/Programs/Online/OnlineTestsCore/Source/Private/Tests/DelegateAdapter/DelegateAdapterTest.cpp
+++ b/Engine/Source/Programs/Online/OnlineTestsCore/Source/Private/Tests/DelegateAdapter/DelegateAdapterTest.cpp
@@ -1,9 +1,5 @@
 // Copyright Epic Games, Inc. All Rights Reserved.
 
-#include "CoreMinimal.h"
-
-#include <catch2/catch_test_macros.hpp>
-
 #include "OnlineCatchHelper.h"
 #include "Online/DelegateAdapter.h"
 #include "Online/MulticastAdapter.h"
diff --git a/Engine/Source/Programs/Online/OnlineTestsCore/Source/Private/Tests/Leaderboards/LeaderboardsReadEntriesAroundRankTests.cpp b/Engine/Source/Programs/Online/OnlineTestsCore/Source/Private/Tests/Leaderboards/LeaderboardsReadEntriesAroundRankTests.cpp
--- a/Engine/Source/Programs/Online/OnlineTestsCore/Source/Private/Tests/Leaderboards/LeaderboardsReadEntriesAroundRankTests.cpp
+++ b/Engine/Source/Programs/Online/OnlineTestsCore/Source/Private/Tests/Leaderboards/LeaderboardsReadEntriesAroundRankTests.cpp
@@ -16,7 +16,7 @@ LEADERBOARDS_TEST_CASE("Verify ReadEntriesAroundRank succeeds and returns entrie
 	TArray< FAccountId*> AccountIds = { &AccountIdA, &AccountIdB };
 	uint32 NumUsersToLogin = 2;
 
-	LeaderboardsFixture_NUsers(std::move(GetLoginPipeline(AccountIdA, AccountIdB)), AccountIds, NumUsersToLogin)
+	LeaderboardsFixture_NUsers(MoveTemp(GetLoginPipeline({ AccountIdA, AccountIdB })), AccountIds, NumUsersToLogin)
 		.EmplaceAsyncLambda([&](FAsyncLambdaResult Promise, SubsystemType Services) { ReadEntriesAroundRank_2UsersFixture(Promise, Services, "Stat_Use_Set", AccountIdA, 0, 2, LeaderboardEntries); })
 		.EmplaceLambda([&](SubsystemType Services) { CheckEntries_2Users_Stat_Use_Set(LeaderboardEntries, AccountIdA, AccountIdB); })
 		.EmplaceAsyncLambda([&](FAsyncLambdaResult Promise, SubsystemType Services) { ReadEntriesAroundRank_2UsersFixture(Promise, Services, "Stat_Use_Smallest", AccountIdA, 0, 2, LeaderboardEntries); })
@@ -32,7 +32,7 @@ LEADERBOARDS_TEST_CASE("Verify ReadEntriesAroundRank returns a fail message if t
 	FAccountId LocalUserInvalid;
 	TArray< FAccountId*> AccountIds = { &AccountId, &LocalUserInvalid };
 
-	LeaderboardsFixture_NUsers(std::move(GetLoginPipeline(AccountId)), AccountIds, 1)
+	LeaderboardsFixture_NUsers(MoveTemp(GetLoginPipeline({ AccountId })), AccountIds, 1)
 		.EmplaceAsyncLambda([&](FAsyncLambdaResult Promise, SubsystemType Services) { ReadEntriesAroundRankFixture(Promise, Services, "Stat_Use_Set", LocalUserInvalid, 0, 2, UE::Online::Errors::InvalidUser(), LeaderboardEntries); });
 
 	RunToCompletion();
@@ -45,7 +45,7 @@ LEADERBOARDS_TEST_CASE("Verify ReadEntriesAroundRank returns a fail message if g
 	FAccountId AccountId;
 	TArray< FAccountId*> AccountIds = { &AccountId };
 
-	LeaderboardsFixture_NUsers(std::move(GetLoginPipeline(AccountId)), AccountIds, 1)
+	LeaderboardsFixture_NUsers(MoveTemp(GetLoginPipeline({ AccountId })), AccountIds, 1)
 		.EmplaceAsyncLambda([&](FAsyncLambdaResult Promise, SubsystemType Services) { ReadEntriesAroundRankFixture(Promise, Services, "Stat_Use_Set", AccountId, 1, 2, UE::Online::Errors::InvalidParams(), LeaderboardEntries); });
 
 	RunToCompletion();
@@ -60,7 +60,7 @@ LEADERBOARDS_TEST_CASE("Verify ReadEntriesAroundRank returns all valid entries i
 	uint32 rank = 0;
 	uint32 limit = 3;
 
-	LeaderboardsFixture_NUsers(std::move(GetLoginPipeline(AccountIdA, AccountIdB)), AccountIds, 2)
+	LeaderboardsFixture_NUsers(MoveTemp(GetLoginPipeline({ AccountIdA, AccountIdB })), AccountIds, 2)
 		.EmplaceAsyncLambda([&](FAsyncLambdaResult Promise, SubsystemType Services) { ReadEntriesAroundRank_2UsersFixture(Promise, Services, "Stat_Use_Set", AccountIdA, rank, limit, LeaderboardEntries); })
 		.EmplaceLambda([&](SubsystemType Services) { CheckEntries_2Users_Stat_Use_Set(LeaderboardEntries, AccountIdA, AccountIdB); });
 
@@ -74,7 +74,7 @@ LEADERBOARDS_TEST_CASE("Verify ReadEntriesAroundRank returns fail message if giv
 	FAccountId AccountId;
 	TArray< FAccountId*> AccountIds = { &AccountId };
 
-	LeaderboardsFixture_NUsers(std::move(GetLoginPipeline(AccountId)), AccountIds, 1)
+	LeaderboardsFixture_NUsers(MoveTemp(GetLoginPipeline({ AccountId })), AccountIds, 1)
 		.EmplaceAsyncLambda([&](FAsyncLambdaResult Promise, SubsystemType Services) { ReadEntriesAroundRankFixture(Promise, Services, "Stat_Use_Set", AccountId, 3, 0, UE::Online::Errors::InvalidParams(), LeaderboardEntries); });
 
 	RunToCompletion();
diff --git a/Engine/Source/Programs/Online/OnlineTestsCore/Source/Private/Tests/Leaderboards/LeaderboardsReadEntriesAroundUserTests.cpp b/Engine/Source/Programs/Online/OnlineTestsCore/Source/Private/Tests/Leaderboards/LeaderboardsReadEntriesAroundUserTests.cpp
--- a/Engine/Source/Programs/Online/OnlineTestsCore/Source/Private/Tests/Leaderboards/LeaderboardsReadEntriesAroundUserTests.cpp
+++ b/Engine/Source/Programs/Online/OnlineTestsCore/Source/Private/Tests/Leaderboards/LeaderboardsReadEntriesAroundUserTests.cpp
@@ -19,7 +19,7 @@ LEADERBOARDS_TEST_CASE("Verify that ReadEntriesAroundUser succeeds and returns e
 	int32 Offset = 0;
 	uint32 Limit = 2;
 
-	LeaderboardsFixture_NUsers(std::move(GetLoginPipeline(AccountIdA, AccountIdB)), AccountIds, 2)
+	LeaderboardsFixture_NUsers(MoveTemp(GetLoginPipeline({ AccountIdA, AccountIdB })), AccountIds, 2)
 		.EmplaceAsyncLambda([&](FAsyncLambdaResult Promise, SubsystemType Services) { ReadEntriesAroundUser_2UsersFixture(Promise, Services, "Stat_Use_Set", AccountIdB, AccountIdA, Offset, Limit, LeaderboardEntries); })
 		.EmplaceLambda([&](SubsystemType Services) { CheckEntries_2Users_Stat_Use_Set(LeaderboardEntries, AccountIdA, AccountIdB); })
 		.EmplaceAsyncLambda([&](FAsyncLambdaResult Promise, SubsystemType Services) { ReadEntriesAroundUser_2UsersFixture(Promise, Services, "Stat_Use_Smallest", AccountIdA, AccountIdB, -1, 2, LeaderboardEntries); })
@@ -37,7 +37,7 @@ LEADERBOARDS_TEST_CASE("Verify ReadEntriesAroundUser returns a fail message if t
 	int32 offset = 0;
 	uint32 limit = 2;
 
-	LeaderboardsFixture_NUsers(std::move(GetLoginPipeline(AccountId)), AccountIds, 1)
+	LeaderboardsFixture_NUsers(MoveTemp(GetLoginPipeline({ AccountId })), AccountIds, 1)
 		.EmplaceAsyncLambda([&](FAsyncLambdaResult Promise, SubsystemType Services) {ReadEntriesAroundUserFixture(Promise, Services, "Stat_Use_Set", LocalUserInvalid, AccountId, offset, limit, UE::Online::Errors::InvalidUser(), LeaderboardEntries); });
 
 	RunToCompletion();
@@ -52,7 +52,7 @@ LEADERBOARDS_TEST_CASE("Verify ReadEntriesAroundUser returns a fail message if t
 	int32 offset = 0;
 	uint32 limit = 2;
 
-	LeaderboardsFixture_NUsers(std::move(GetLoginPipeline(LocalUser)), AccountIds, 1)
+	LeaderboardsFixture_NUsers(MoveTemp(GetLoginPipeline({ LocalUser })), AccountIds, 1)
 		.EmplaceAsyncLambda([&](FAsyncLambdaResult Promise, SubsystemType Services) {ReadEntriesAroundUserFixture(Promise, Services, "Stat_Use_Set", LocalUser, AccountIdInvalid, offset, limit, UE::Online::Errors::InvalidUser(), LeaderboardEntries); });
 
 	RunToCompletion();
@@ -66,7 +66,7 @@ LEADERBOARDS_TEST_CASE("Verify ReadEntriesAround User returns entries centered a
 	int32 offset = 1;
 	uint32 limit = 2;
 
-	LeaderboardsFixture_NUsers(std::move(GetLoginPipeline(LocalUser, OtherPlayerB, UserPlayer, OtherPlayerC, OtherPlayerD)), AccountIds, 5)
+	LeaderboardsFixture_NUsers(MoveTemp(GetLoginPipeline({ LocalUser, OtherPlayerB, UserPlayer, OtherPlayerC, OtherPlayerD })), AccountIds, 5)
 		.EmplaceAsyncLambda([&](FAsyncLambdaResult Promise, SubsystemType Services) { ReadEntriesAroundUser_5UsersFixture(Promise, Services, "Stat_Use_Set", { LocalUser, OtherPlayerB, UserPlayer, OtherPlayerC, OtherPlayerD }, offset, limit, LeaderboardEntries); })
 		.EmplaceLambda([&](SubsystemType Services) { CheckEntries_2UsersRanks_Stat_Use_Set(LeaderboardEntries, OtherPlayerC, OtherPlayerD); });
 
@@ -81,7 +81,7 @@ LEADERBOARDS_TEST_CASE("Verify ReadEntriesAroundUser returns a number of entries
 	int32 offset = 0;
 	uint32 limit = 2;
 
-	LeaderboardsFixture_NUsers(std::move(GetLoginPipeline(LocalUser, AccountId)), AccountIds, 2)
+	LeaderboardsFixture_NUsers(MoveTemp(GetLoginPipeline({ LocalUser, AccountId })), AccountIds, 2)
 		.EmplaceAsyncLambda([&](FAsyncLambdaResult Promise, SubsystemType Services) {
 		Services->GetLeaderboardsInterface()->ReadEntriesAroundUser({ LocalUser, AccountId, offset, limit, "Stat_Use_Set" })
 			.OnComplete([Promise, limit](const TOnlineResult<UE::Online::FReadEntriesAroundUser>& Result) mutable
@@ -104,7 +104,7 @@ LEADERBOARDS_TEST_CASE("Verify ReadEntriesAroundUser returns expected valid entr
 	uint32 limit = 4;
 	TArray<UE::Online::FLeaderboardEntry> LeaderboardEntries;
 
-	LeaderboardsFixture_NUsers(std::move(GetLoginPipeline(AccountId, LocalUser)), AccountIds, 2)
+	LeaderboardsFixture_NUsers(MoveTemp(GetLoginPipeline({ AccountId, LocalUser })), AccountIds, 2)
 		.EmplaceAsyncLambda([&](FAsyncLambdaResult Promise, SubsystemType Services) { ReadEntriesAroundUser_2UsersFixture(Promise, Services, "Stat_Use_Set", LocalUser, AccountId, offset, limit, LeaderboardEntries); })
 		.EmplaceLambda([&](SubsystemType Services) { CheckEntries_2Users_Stat_Use_Set(LeaderboardEntries, AccountId, LocalUser); });
 
@@ -121,7 +121,7 @@ LEADERBOARDS_TEST_CASE("Verify ReadEntriesAroundUser returns expected valid entr
 	uint32 limit = 8;
 	TArray<UE::Online::FLeaderboardEntry> LeaderboardEntries;
 
-	LeaderboardsFixture_NUsers(std::move(GetLoginPipeline(LocalUser, AccountId)), AccountIds, 2)
+	LeaderboardsFixture_NUsers(MoveTemp(GetLoginPipeline({ LocalUser, AccountId })), AccountIds, 2)
 		.EmplaceAsyncLambda([&](FAsyncLambdaResult Promise, SubsystemType Services) { ReadEntriesAroundUser_2UsersFixture(Promise, Services, "Stat_Use_Smallest", LocalUser, AccountId, offset, limit, LeaderboardEntries); })
 		.EmplaceLambda([&](SubsystemType Services) { CheckEntries_1User_Stat_Use_Smallest(LeaderboardEntries, AccountId); });
 
@@ -136,7 +136,7 @@ LEADERBOARDS_TEST_CASE("Verify ReadEntriesAroundUser returns error message if gi
 	uint32 limit = 0;
 	TArray<UE::Online::FLeaderboardEntry> LeaderboardEntries;
 
-	LeaderboardsFixture_NUsers(std::move(GetLoginPipeline(LocalUser, AccountId)), AccountIds, 2)
+	LeaderboardsFixture_NUsers(MoveTemp(GetLoginPipeline({ LocalUser, AccountId })), AccountIds, 2)
 		.EmplaceAsyncLambda([&](FAsyncLambdaResult Promise, SubsystemType Services) {ReadEntriesAroundUserFixture(Promise, Services, "Stat_Use_Set", LocalUser, AccountId, offset, limit, UE::Online::Errors::InvalidParams(), LeaderboardEntries); });
 
 	RunToCompletion();
@@ -149,7 +149,7 @@ LEADERBOARDS_TEST_CASE("Verify ReadEntriesAroundUser returns only one entry if g
 	int32 offset = 0;
 	uint32 limit = 1;
 
-	LeaderboardsFixture_NUsers(std::move(GetLoginPipeline(LocalUser, AccountId)), AccountIds, 2)
+	LeaderboardsFixture_NUsers(MoveTemp(GetLoginPipeline({ LocalUser, AccountId })), AccountIds, 2)
 		.EmplaceAsyncLambda([&](FAsyncLambdaResult Promise, SubsystemType Services) {
 		Services->GetLeaderboardsInterface()->ReadEntriesAroundUser({ LocalUser, AccountId, offset, limit, "Stat_Use_Set" })
 			.OnComplete([Promise, limit](const TOnlineResult<UE::Online::FReadEntriesAroundUser>& Result) mutable
@@ -170,7 +170,7 @@ LEADERBOARDS_TEST_CASE("Verify ReadEntriesAroundUser returns expected valid entr
 	int32 offset = 0;
 	uint32 limit = 2;
 
-	LeaderboardsFixture_NUsers(std::move(GetLoginPipeline(OtherPlayerA, OtherPlayerB, UserPlayer, OtherPlayerC, OtherPlayerD)), AccountIds, 5)
+	LeaderboardsFixture_NUsers(MoveTemp(GetLoginPipeline({ OtherPlayerA, OtherPlayerB, UserPlayer, OtherPlayerC, OtherPlayerD })), AccountIds, 5)
 		.EmplaceAsyncLambda([&](FAsyncLambdaResult Promise, SubsystemType Services) { ReadEntriesAroundUser_5UsersFixture(Promise, Services, "Stat_Use_Set", { OtherPlayerA, OtherPlayerB, UserPlayer, OtherPlayerC, OtherPlayerD }, offset, limit, LeaderboardEntries); })
 		.EmplaceLambda([&](SubsystemType Services) { CheckEntries_2UsersRanks_Stat_Use_Set_Index_Moved_Up(LeaderboardEntries, UserPlayer, OtherPlayerC); });
 
diff --git a/Engine/Source/Programs/Online/OnlineTestsCore/Source/Private/Tests/Leaderboards/LeaderboardsReadEntriesForUsersTests.cpp b/Engine/Source/Programs/Online/OnlineTestsCore/Source/Private/Tests/Leaderboards/LeaderboardsReadEntriesForUsersTests.cpp
--- a/Engine/Source/Programs/Online/OnlineTestsCore/Source/Private/Tests/Leaderboards/LeaderboardsReadEntriesForUsersTests.cpp
+++ b/Engine/Source/Programs/Online/OnlineTestsCore/Source/Private/Tests/Leaderboards/LeaderboardsReadEntriesForUsersTests.cpp
@@ -17,7 +17,7 @@ LEADERBOARDS_TEST_CASE("ReadEntriesForUsers succeed", "[MultiAccount]")
 	FAccountId AccountIdA, AccountIdB;
 	TArray< FAccountId*> AccountIds = { &AccountIdA, &AccountIdB };
 
-	LeaderboardsFixture_NUsers(std::move(GetLoginPipeline(AccountIdA, AccountIdB)), AccountIds, 2)
+	LeaderboardsFixture_NUsers(MoveTemp(GetLoginPipeline({ AccountIdA, AccountIdB })), AccountIds, 2)
 		.EmplaceAsyncLambda([&](FAsyncLambdaResult Promise, SubsystemType Services) { ReadEntriesForUsers_2UsersFixture(Promise, Services, "Stat_Use_Set", AccountIdA, AccountIdB, LeaderboardEntries); })
 		.EmplaceLambda([&](SubsystemType Services) { CheckEntries_2Users_Stat_Use_Set(LeaderboardEntries, AccountIdA, AccountIdB); })
 		.EmplaceAsyncLambda([&](FAsyncLambdaResult Promise, SubsystemType Services) { ReadEntriesForUsers_2UsersFixture(Promise, Services, "Stat_Use_Smallest", AccountIdA, AccountIdB, LeaderboardEntries); })
@@ -32,7 +32,7 @@ LEADERBOARDS_TEST_CASE("Verify ReadEntriesForUsers returns a fail message if the
 	FAccountId AccountId;	
 	FAccountId LocalUserInvalid;
 	TArray< FAccountId*> AccountIds = { &AccountId, &LocalUserInvalid };
-	LeaderboardsFixture_NUsers(std::move(GetLoginPipeline(AccountId)), AccountIds, 2)
+	LeaderboardsFixture_NUsers(MoveTemp(GetLoginPipeline({ AccountId })), AccountIds, 2)
 		.EmplaceAsyncLambda([&](FAsyncLambdaResult Promise, SubsystemType Services) { ReadEntriesForUsersFixture(Promise, Services, "Stat_Use_Set", { LocalUserInvalid, AccountId }, UE::Online::Errors::InvalidUser(), LeaderboardEntries); });
 
 	RunToCompletion();
@@ -43,7 +43,7 @@ LEADERBOARDS_TEST_CASE("Verify ReadEntriesForUsers returns an empty array of ent
 	TArray<UE::Online::FLeaderboardEntry> LeaderboardEntries;
 	FAccountId LocalUser;
 	TArray< FAccountId*> AccountIds = { &LocalUser };
-	LeaderboardsFixture_NUsers(std::move(GetLoginPipeline(LocalUser)), AccountIds, 1)
+	LeaderboardsFixture_NUsers(MoveTemp(GetLoginPipeline({ LocalUser })), AccountIds, 1)
 		.EmplaceAsyncLambda([&](FAsyncLambdaResult Promise, SubsystemType Services) {
 		Services->GetLeaderboardsInterface()->ReadEntriesForUsers({ LocalUser, {}, "Stat_Use_Set" })
 		.OnComplete([Promise, &LeaderboardEntries](const TOnlineResult<UE::Online::FReadEntriesForUsers>& Result) mutable
@@ -61,7 +61,7 @@ LEADERBOARDS_TEST_CASE("Verify ReadEntiresForUsers returns a fail message if giv
 	FAccountId AccountIdInvalid;
 	FAccountId LocalUser;
 	TArray< FAccountId*> AccountIds = { &LocalUser, &AccountIdInvalid };
-	LeaderboardsFixture_NUsers(std::move(GetLoginPipeline(LocalUser)), AccountIds, 1)
+	LeaderboardsFixture_NUsers(MoveTemp(GetLoginPipeline({ LocalUser })), AccountIds, 1)
 		.EmplaceAsyncLambda([&](FAsyncLambdaResult Promise, SubsystemType Services) {
 		Services->GetLeaderboardsInterface()->ReadEntriesForUsers({ LocalUser, {AccountIdInvalid}, "Stat_Use_Set" })
 		.OnComplete([Promise, &LeaderboardEntries](const TOnlineResult<UE::Online::FReadEntriesForUsers>& Result) mutable
@@ -87,7 +87,7 @@ LEADERBOARDS_TEST_CASE("Verify ReadEntriesForUsers returns a fail message if giv
 	FAccountId AccountIdInvalid;
 	FAccountId LocalUser;
 	TArray< FAccountId*> AccountIds = { &LocalUser, &AccountIdInvalid };
-	LeaderboardsFixture_NUsers(std::move(GetLoginPipeline(LocalUser)), AccountIds, 1)
+	LeaderboardsFixture_NUsers(MoveTemp(GetLoginPipeline({ LocalUser })), AccountIds, 1)
 		.EmplaceAsyncLambda([&](FAsyncLambdaResult Promise, SubsystemType Services) { ReadEntriesForUsersFixture(Promise, Services, "Stat_Use_Set", { LocalUser, AccountIdInvalid }, UE::Online::Errors::InvalidUser(), LeaderboardEntries); });
 	RunToCompletion();
 }
@@ -97,7 +97,7 @@ LEADERBOARDS_TEST_CASE("Verify ReadEntriesForUsers returns an Entries array with
 	TArray<UE::Online::FLeaderboardEntry> LeaderboardEntries;
 	FAccountId LocalUser;
 	TArray< FAccountId*> AccountIds = { &LocalUser };
-	LeaderboardsFixture_NUsers(std::move(GetLoginPipeline(LocalUser)), AccountIds, 1)
+	LeaderboardsFixture_NUsers(MoveTemp(GetLoginPipeline({ LocalUser })), AccountIds, 1)
 		.EmplaceAsyncLambda([&](FAsyncLambdaResult Promise, SubsystemType Services) {
 		Services->GetLeaderboardsInterface()->ReadEntriesForUsers({ LocalUser, {LocalUser}, "Stat_Use_Set" })
 		.OnComplete([Promise, &LeaderboardEntries, &LocalUser](const TOnlineResult<UE::Online::FReadEntriesForUsers>& Result) mutable
diff --git a/Engine/Source/Programs/Online/OnlineTestsCore/Source/Private/Tests/Lobby/LobbyCreateInviteTests.cpp b/Engine/Source/Programs/Online/OnlineTestsCore/Source/Private/Tests/Lobby/LobbyCreateInviteTests.cpp
--- a/Engine/Source/Programs/Online/OnlineTestsCore/Source/Private/Tests/Lobby/LobbyCreateInviteTests.cpp
+++ b/Engine/Source/Programs/Online/OnlineTestsCore/Source/Private/Tests/Lobby/LobbyCreateInviteTests.cpp
@@ -1,20 +1,16 @@
 // Copyright Epic Games, Inc. All Rights Reserved.
 
-#include "CoreMinimal.h"
-
-#include <catch2/catch_test_macros.hpp>
-
-#include "Helpers/Identity/IdentityGetLoginByUserId.h"
 #include "Helpers/Lobby/LobbyCreateHelper.h"
 #include "Helpers/Lobby/LobbyJoinHelper.h"
-#include "OnlineCatchHelper.h"
 
 #define LOBBY_CREATE_INVITE_TAGS "[Lobby]"
 #define LOBBY_CREATE_INVITE_TEST_CASE(x, ...) ONLINE_TEST_CASE(x, LOBBY_CREATE_INVITE_TAGS __VA_ARGS__)
 
 LOBBY_CREATE_INVITE_TEST_CASE("Basic lobby create and join test")
 {
-	int32 LocalUserNum = 0;
+	int32 UserNumToLogin = 1;
+	FAccountId AccountId;
+
 	UE::Online::FCreateLobby::Params LobbyCreateParams;
 	LobbyCreateParams.LocalName = TEXT("TestLobby");
 	LobbyCreateParams.SchemaId = TEXT("test");
@@ -25,14 +21,16 @@ LOBBY_CREATE_INVITE_TEST_CASE("Basic lobby create and join test")
 	FJoinLobby::Params JoinLobbyParams;
 	JoinLobbyParams.LocalName = TEXT("TestLobby");
 
-	GetLoginPipeline(1)
-		.EmplaceStep<FIdentityGetUniquePlayerIdStep>(LocalUserNum, 
-			[&LobbyCreateParams, &JoinLobbyParams](UE::Online::FAccountInfo InAccountInfo) {
-				LobbyCreateParams.LocalAccountId = InAccountInfo.AccountId;
-				JoinLobbyParams.LocalAccountId = InAccountInfo.AccountId;
-			})
-		.EmplaceStep<FLobbyCreateHelper>(&LobbyCreateParams, [&JoinLobbyParams](UE::Online::FLobby InLobby) {
-				JoinLobbyParams.LobbyId = InLobby.LobbyId; }, true);
+	FTestPipeline& LoginPipeline = GetLoginPipeline(UserNumToLogin, { AccountId });
+
+	LobbyCreateParams.LocalAccountId = AccountId;
+	JoinLobbyParams.LocalAccountId = AccountId;
+	
+	LoginPipeline
+		.EmplaceStep<FLobbyCreateHelper>(&LobbyCreateParams, [&JoinLobbyParams](UE::Online::FLobby InLobby) 
+			{
+				JoinLobbyParams.LobbyId = InLobby.LobbyId; 
+			}, true);
 
 	RunToCompletion();
 }
diff --git a/Engine/Source/Programs/Online/OnlineTestsCore/Source/Private/Tests/Meta/AsyncTestStepTest.cpp b/Engine/Source/Programs/Online/OnlineTestsCore/Source/Private/Tests/Meta/AsyncTestStepTest.cpp
--- a/Engine/Source/Programs/Online/OnlineTestsCore/Source/Private/Tests/Meta/AsyncTestStepTest.cpp
+++ b/Engine/Source/Programs/Online/OnlineTestsCore/Source/Private/Tests/Meta/AsyncTestStepTest.cpp
@@ -1,10 +1,6 @@
 // Copyright Epic Games, Inc. All Rights Reserved.
 
 #include "AsyncTestStep.h"
-#include "CoreMinimal.h"
-
-#include <catch2/catch_test_macros.hpp>
-
 #include "OnlineCatchHelper.h"
 #include "Online/DelegateAdapter.h"
 #include "Online/MulticastAdapter.h"
@@ -37,7 +33,7 @@ AST_ERROR_TEST_CASE("Basic test- confirm all of async steps, lambdas, and async
 	bool bHasRun = false;
 	bool bDidComplete = false;
 
-	GetLoginPipeline(0)
+	GetPipeline()
 		.EmplaceStep<FAsyncTestStepHelper>(bHasRun)
 		.EmplaceLambda([&](SubsystemType Type)
 		{
diff --git a/Engine/Source/Programs/Online/OnlineTestsCore/Source/Private/Tests/Meta/DemoTest.cpp b/Engine/Source/Programs/Online/OnlineTestsCore/Source/Private/Tests/Meta/DemoTest.cpp
--- a/Engine/Source/Programs/Online/OnlineTestsCore/Source/Private/Tests/Meta/DemoTest.cpp
+++ b/Engine/Source/Programs/Online/OnlineTestsCore/Source/Private/Tests/Meta/DemoTest.cpp
@@ -1,6 +1,5 @@
 // Copyright Epic Games, Inc. All Rights Reserved.
-#include "CoreMinimal.h"
-#include "TestHarness.h"
+
 #include "OnlineCatchHelper.h"
  
 #define DEMOTEST_TAG "[Demo]"  
diff --git a/Engine/Source/Programs/Online/OnlineTestsCore/Source/Private/Tests/Meta/MetaTests.cpp b/Engine/Source/Programs/Online/OnlineTestsCore/Source/Private/Tests/Meta/MetaTests.cpp
--- a/Engine/Source/Programs/Online/OnlineTestsCore/Source/Private/Tests/Meta/MetaTests.cpp
+++ b/Engine/Source/Programs/Online/OnlineTestsCore/Source/Private/Tests/Meta/MetaTests.cpp
@@ -1,8 +1,6 @@
 // Copyright Epic Games, Inc. All Rights Reserved.
 
-#include <catch2/catch_test_macros.hpp>
 #include "AsyncTestStep.h"
-#include "CoreMinimal.h"
 #include "Online/DelegateAdapter.h"
 #include "Online/MulticastAdapter.h"
 #include "OnlineCatchHelper.h"
@@ -18,7 +16,7 @@ META_TEST_CASE("Async test steps")
 	{
 		bool bDidComplete = false;
 
-		GetLoginPipeline(0)
+		GetPipeline()
 			.EmplaceAsyncLambda([&](FAsyncLambdaResult Result, SubsystemType Type)
 				{
 					// This should cause a failure and crash, but if the fix is working
diff --git a/Engine/Source/Programs/Online/OnlineTestsCore/Source/Private/Tests/Meta/SyncLoginTest.cpp b/Engine/Source/Programs/Online/OnlineTestsCore/Source/Private/Tests/Meta/SyncLoginTest.cpp
--- a/Engine/Source/Programs/Online/OnlineTestsCore/Source/Private/Tests/Meta/SyncLoginTest.cpp
+++ b/Engine/Source/Programs/Online/OnlineTestsCore/Source/Private/Tests/Meta/SyncLoginTest.cpp
@@ -1,7 +1,5 @@
 // Copyright Epic Games, Inc. All Rights Reserved.
 
-#include "CoreMinimal.h"
-#include <catch2/catch_test_macros.hpp>
 #include "OnlineCatchHelper.h"
 
 #define SYNCLOGIN_ERROR_TAG "[NULL][Meta][SyncLogin]"
@@ -11,7 +9,7 @@ SYNCLOGIN_ERROR_TEST_CASE("Confirm sync login is logging in before RunToCompleti
 {
 	FAccountId AccountId, AccountId2;
 
-	auto& PipeSave = GetLoginPipeline(AccountId, AccountId2);
+	auto& PipeSave = GetLoginPipeline({ AccountId, AccountId2 });
 	// Check is valid right after Pipeline call.
 	CHECK(AccountId.IsValid());
 	CHECK(AccountId2.IsValid());
@@ -39,7 +37,7 @@ SYNCLOGIN_ERROR_TEST_CASE("Confirm sync login is logging in before RunToCompleti
 	FAccountId AccountId, AccountId2, AccountId3, AccountId4, AccountId5;
 	bool bHasRun = false;
 
-	GetLoginPipeline(AccountId, AccountId2, AccountId3, AccountId4, AccountId5)
+	GetLoginPipeline({ AccountId, AccountId2, AccountId3, AccountId4, AccountId5 })
 		.EmplaceLambda([&](SubsystemType Type)
 		{
 				// Check valid during pipeline
diff --git a/Engine/Source/Programs/Online/OnlineTestsCore/Source/Private/Tests/Presence/PresenceTests.cpp b/Engine/Source/Programs/Online/OnlineTestsCore/Source/Private/Tests/Presence/PresenceTests.cpp
--- a/Engine/Source/Programs/Online/OnlineTestsCore/Source/Private/Tests/Presence/PresenceTests.cpp
+++ b/Engine/Source/Programs/Online/OnlineTestsCore/Source/Private/Tests/Presence/PresenceTests.cpp
@@ -1,9 +1,5 @@
 // Copyright Epic Games, Inc. All Rights Reserved.
 
-#include "CoreMinimal.h"
-#include <catch2/catch_test_macros.hpp>
-
-#include "OnlineCatchHelper.h"
 #include "Helpers/Presence/BatchQueryPresenceHelper.h"
 #include "Helpers/Presence/GetCachedPresenceHelper.h"
 #include "Helpers/Presence/PartialUpdatePresenceHelper.h"
@@ -19,7 +15,7 @@ PRESENCE_TEST_CASE("Ensure UpdatePresence works as expected and can then be quer
 	FAccountId AccountIdA;
 	TSharedPtr<const FUserPresence> OutPresence;
 
-	GetLoginPipeline(AccountIdA)
+	GetLoginPipeline({ AccountIdA })
 		.EmplaceStep<FUpdatePresenceHelper>(AccountIdA, GetGenericPresenceA(AccountIdA))
 		.EmplaceStep<FQueryPresenceHelper>(AccountIdA, AccountIdA, OutPresence)
 		.EmplaceStep<FComparePresencesHelper>(OutPresence, GetGenericPresenceA(AccountIdA));
@@ -33,7 +29,7 @@ PRESENCE_TEST_CASE("Ensure UpdatePresence works as expected and can then be quer
 	FAccountId AccountIdA, AccountIdB;
 	TSharedPtr<const FUserPresence> OutPresence;
 
-	GetLoginPipeline(AccountIdA, AccountIdB)
+	GetLoginPipeline({ AccountIdA, AccountIdB })
 		// <todo: ensure friendship helper>
 		.EmplaceStep<FUpdatePresenceHelper>(AccountIdB, GetGenericPresenceA(AccountIdB))
 		.EmplaceStep<FQueryPresenceHelper>(AccountIdA, AccountIdB, OutPresence)
@@ -48,7 +44,7 @@ PRESENCE_TEST_CASE("Ensure UpdatePresence works on multiple users and can then b
 	TArray<FAccountId> AccountIds;
 	TArray<TSharedRef<const FUserPresence>> OutPresences;
 
-	GetLoginPipeline(AccountIdA, AccountIdB)
+	GetLoginPipeline({ AccountIdA, AccountIdB })
 		.EmplaceLambda([&](SubsystemType Type)
 		{
 			AccountIds.Add(AccountIdA);
@@ -73,7 +69,7 @@ PRESENCE_TEST_CASE("Ensure PartialUpdatePresence works as expected and can then
 	FAccountId AccountIdA, AccountIdB;
 	TSharedPtr<const FUserPresence> OutPresence;
 
-	GetLoginPipeline(AccountIdA, AccountIdB)
+	GetLoginPipeline({ AccountIdA, AccountIdB })
 		// <todo: ensure friendship helper>
 		.EmplaceStep<FUpdatePresenceHelper>(AccountIdA, GetGenericPresenceA(AccountIdA))
 		.EmplaceStep<FPartialUpdatePresenceHelper>(GetGenericPresenceMutationAtoB(AccountIdB))
@@ -88,7 +84,7 @@ PRESENCE_TEST_CASE("Ensure UpdatePresence works as expected and can then be quer
 	FAccountId AccountIdA;
 	TSharedPtr<const FUserPresence> OutPresenceA, OutPresenceB;
 
-	GetLoginPipeline(AccountIdA)
+	GetLoginPipeline({ AccountIdA })
 		.EmplaceStep<FUpdatePresenceHelper>(AccountIdA, GetGenericPresenceA(AccountIdA))
 		.EmplaceStep<FQueryPresenceHelper>(AccountIdA, AccountIdA, OutPresenceA)
 		.EmplaceStep<FGetCachedPresenceHelper>(AccountIdA, AccountIdA, OutPresenceB)
@@ -116,7 +112,7 @@ PRESENCE_TEST_CASE("Ensure QueryPresence's bListenToChanges mod fires the event
 		}
 	});
 
-	GetLoginPipeline(AccountIdA, AccountIdB)
+	GetLoginPipeline({ AccountIdA, AccountIdB })
 		// <todo: ensure friendship helper>
 		.EmplaceStep<FUpdatePresenceHelper>(AccountIdB, GetGenericPresenceA(AccountIdB))
 		.EmplaceStep<FQueryPresenceHelper>(AccountIdA, AccountIdB, OutPresence, true)
@@ -149,7 +145,7 @@ PRESENCE_TEST_CASE("Ensure BatchQueryPresence's bListenToChanges mod fires the e
 		}
 	});
 
-	GetLoginPipeline(AccountIdA, AccountIdB)
+	GetLoginPipeline({ AccountIdA, AccountIdB })
 		.EmplaceLambda([&](SubsystemType Type)
 		{
 			AccountIds.Add(AccountIdA);
diff --git a/Engine/Source/Programs/Online/OnlineTestsCore/Source/Private/Tests/Schema/SchemaTests.cpp b/Engine/Source/Programs/Online/OnlineTestsCore/Source/Private/Tests/Schema/SchemaTests.cpp
--- a/Engine/Source/Programs/Online/OnlineTestsCore/Source/Private/Tests/Schema/SchemaTests.cpp
+++ b/Engine/Source/Programs/Online/OnlineTestsCore/Source/Private/Tests/Schema/SchemaTests.cpp
@@ -1,9 +1,5 @@
 // Copyright Epic Games, Inc. All Rights Reserved.
 
-#include "CoreMinimal.h"
-
-#include <catch2/catch_test_macros.hpp>
-
 #include "OnlineCatchHelper.h"
 #include "Online/Schema.h"
 
diff --git a/Engine/Source/Programs/Online/OnlineTestsCore/Source/Private/Tests/Sessions/AddRemoveSessionMemberTests.cpp b/Engine/Source/Programs/Online/OnlineTestsCore/Source/Private/Tests/Sessions/AddRemoveSessionMemberTests.cpp
--- a/Engine/Source/Programs/Online/OnlineTestsCore/Source/Private/Tests/Sessions/AddRemoveSessionMemberTests.cpp
+++ b/Engine/Source/Programs/Online/OnlineTestsCore/Source/Private/Tests/Sessions/AddRemoveSessionMemberTests.cpp
@@ -2,7 +2,7 @@
 
 #include "Helpers/Sessions/CreateSessionHelper.h"
 #include "Helpers/Sessions/AddRemoveSessionMemberHelper.h"
-#include "OnlineCatchHelper.h"
+#include "Helpers/Sessions/LeaveSessionHelper.h"
 
 #define SESSIONS_TAG "[suite_sessions]"
 #define EG_SESSIONS_ADDREMOVESESSIONMEMBER_TAG SESSIONS_TAG "[addremovesessionmember]"
@@ -10,15 +10,13 @@
 
 SESSIONS_TEST_CASE("If I call AddSessionMember with an invalid account id, I get an error", EG_SESSIONS_ADDREMOVESESSIONMEMBER_TAG)
 {
-	const int32 NumUsersToLogin = 0;
-
 	FAddSessionMember::Params OpAddParams;
 	FAddSessionMemberHelper::FHelperParams AddSessionMemberHelperParams;
 	AddSessionMemberHelperParams.OpParams = &OpAddParams;
 	AddSessionMemberHelperParams.OpParams->LocalAccountId = FAccountId();
 	AddSessionMemberHelperParams.ExpectedError = TOnlineResult<FAddSessionMember>(Errors::InvalidParams());
 
-	GetLoginPipeline(NumUsersToLogin)
+	GetPipeline()
 		.EmplaceStep<FAddSessionMemberHelper>(MoveTemp(AddSessionMemberHelperParams));
 
 	RunToCompletion();
@@ -34,7 +32,7 @@ SESSIONS_TEST_CASE("If I call AddSessionMember with an empty session name, I get
 	AddSessionMemberHelperParams.OpParams->SessionName = TEXT("");
 	AddSessionMemberHelperParams.ExpectedError = TOnlineResult<FAddSessionMember>(Errors::InvalidParams());
 
-	FTestPipeline& LoginPipeline = GetLoginPipeline(AccountId);
+	FTestPipeline& LoginPipeline = GetLoginPipeline({ AccountId });
 	
 	AddSessionMemberHelperParams.OpParams->LocalAccountId = AccountId;
 
@@ -53,7 +51,7 @@ SESSIONS_TEST_CASE("If I call AddSessionMember with an unregistered session name
 	AddSessionMemberHelperParams.OpParams = &OpAddParams;
 	AddSessionMemberHelperParams.ExpectedError = TOnlineResult<FAddSessionMember>(Errors::InvalidState());
 
-	FTestPipeline& LoginPipeline = GetLoginPipeline(AccountId);
+	FTestPipeline& LoginPipeline = GetLoginPipeline({ AccountId });
 
 	AddSessionMemberHelperParams.OpParams->LocalAccountId = AccountId;
 	AddSessionMemberHelperParams.OpParams->SessionName = TEXT("UnregisteredName");
@@ -67,8 +65,7 @@ SESSIONS_TEST_CASE("If I call AddSessionMember with an unregistered session name
 SESSIONS_TEST_CASE("If I call AddSessionMember with valid data, the operation completes successfully", EG_SESSIONS_ADDREMOVESESSIONMEMBER_TAG)
 {
 	DestroyCurrentServiceModule();
-	ResetAccountStatus();
-
+	
 	FAccountId AccountId;
 
 	FCreateSession::Params OpCreateParams;
@@ -84,29 +81,35 @@ SESSIONS_TEST_CASE("If I call AddSessionMember with valid data, the operation co
 	AddSessionMemberHelperParams.OpParams = &OpAddParams;
 	AddSessionMemberHelperParams.OpParams->SessionName = FName(TEXT("SessionValidNameAddMember"));
 
-	FTestPipeline& LoginPipeline = GetLoginPipeline(AccountId);
+	FLeaveSession::Params OpLeaveParams;
+	FLeaveSessionHelper::FHelperParams LeaveSessionHelperParams;
+	LeaveSessionHelperParams.OpParams = &OpLeaveParams;
+	LeaveSessionHelperParams.OpParams->SessionName = TEXT("SessionValidNameAddMember");
+	LeaveSessionHelperParams.OpParams->bDestroySession = true;
+
+	FTestPipeline& LoginPipeline = GetLoginPipeline({ AccountId });
 
 	CreateSessionHelperParams.OpParams->LocalAccountId = AccountId;
 	AddSessionMemberHelperParams.OpParams->LocalAccountId = AccountId;
+	LeaveSessionHelperParams.OpParams->LocalAccountId = AccountId;
 
 	LoginPipeline
 		.EmplaceStep<FCreateSessionHelper>(MoveTemp(CreateSessionHelperParams))
-		.EmplaceStep<FAddSessionMemberHelper>(MoveTemp(AddSessionMemberHelperParams));
+		.EmplaceStep<FAddSessionMemberHelper>(MoveTemp(AddSessionMemberHelperParams))
+		.EmplaceStep<FLeaveSessionHelper>(MoveTemp(LeaveSessionHelperParams));
 
 	RunToCompletion();
 }
 
 SESSIONS_TEST_CASE("If I call RemoveSessionMember with an invalid account id, I get an error", EG_SESSIONS_ADDREMOVESESSIONMEMBER_TAG)
 {
-	const int32 NumUsersToLogin = 0;
-
 	FRemoveSessionMember::Params OpRemoveParams;
 	FRemoveSessionMemberHelper::FHelperParams RemoveSessionMemberHelperParams;
 	RemoveSessionMemberHelperParams.OpParams = &OpRemoveParams;
 	RemoveSessionMemberHelperParams.OpParams->LocalAccountId = FAccountId();
 	RemoveSessionMemberHelperParams.ExpectedError = TOnlineResult<FRemoveSessionMember>(Errors::InvalidParams());
 
-	GetLoginPipeline(NumUsersToLogin)
+	GetPipeline()
 		.EmplaceStep<FRemoveSessionMemberHelper>(MoveTemp(RemoveSessionMemberHelperParams));
 
 	RunToCompletion();
@@ -122,7 +125,7 @@ SESSIONS_TEST_CASE("If I call RemoveSessionMember with an empty session name, I
 	RemoveSessionMemberHelperParams.OpParams->SessionName = TEXT("");
 	RemoveSessionMemberHelperParams.ExpectedError = TOnlineResult<FRemoveSessionMember>(Errors::InvalidParams());
 
-	FTestPipeline& LoginPipeline = GetLoginPipeline(AccountId);
+	FTestPipeline& LoginPipeline = GetLoginPipeline({ AccountId });
 
 	RemoveSessionMemberHelperParams.OpParams->LocalAccountId = AccountId;
 
@@ -141,7 +144,7 @@ SESSIONS_TEST_CASE("If I call RemoveSessionMember with an unregistered session n
 	RemoveSessionMemberHelperParams.OpParams->SessionName = TEXT("UnregisteredName");
 	RemoveSessionMemberHelperParams.ExpectedError = TOnlineResult<FRemoveSessionMember>(Errors::InvalidState());
 
-	FTestPipeline& LoginPipeline = GetLoginPipeline(AccountId);
+	FTestPipeline& LoginPipeline = GetLoginPipeline({ AccountId });
 
 	RemoveSessionMemberHelperParams.OpParams->LocalAccountId = AccountId;
 
@@ -153,15 +156,15 @@ SESSIONS_TEST_CASE("If I call RemoveSessionMember with an unregistered session n
 SESSIONS_TEST_CASE("If I call RemoveSessionMember with valid data, the operation completes successfully", EG_SESSIONS_ADDREMOVESESSIONMEMBER_TAG)
 {
 	DestroyCurrentServiceModule();
-	ResetAccountStatus();
 
+	int32 UserNumToLogin = 7;
 	FAccountId AccountId;
 
 	FCreateSession::Params OpCreateParams;
 	FCreateSessionHelper::FHelperParams CreateSessionHelperParams;
 	CreateSessionHelperParams.OpParams = &OpCreateParams;
 	CreateSessionHelperParams.OpParams->SessionName = TEXT("SessionNameValidRemoveMember");
-	CreateSessionHelperParams.OpParams->SessionSettings.SchemaName = FName(TEXT("SchemaName"));
+	CreateSessionHelperParams.OpParams->SessionSettings.SchemaName = FName(TEXT("SchemaName4"));
 	CreateSessionHelperParams.OpParams->SessionSettings.NumMaxConnections = 4;
 	CreateSessionHelperParams.OpParams->bPresenceEnabled = true;
 
@@ -170,21 +173,29 @@ SESSIONS_TEST_CASE("If I call RemoveSessionMember with valid data, the operation
 	AddSessionMemberHelperParams.OpParams = &OpAddParams;
 	AddSessionMemberHelperParams.OpParams->SessionName = TEXT("SessionNameValidRemoveMember");
 
+	FLeaveSession::Params OpLeaveParams;
+	FLeaveSessionHelper::FHelperParams LeaveSessionHelperParams;
+	LeaveSessionHelperParams.OpParams = &OpLeaveParams;
+	LeaveSessionHelperParams.OpParams->SessionName = TEXT("SessionNameValidRemoveMember");
+	LeaveSessionHelperParams.OpParams->bDestroySession = true;
+
 	FRemoveSessionMember::Params OpRemoveParams;
 	FRemoveSessionMemberHelper::FHelperParams RemoveSessionMemberHelperParams;
 	RemoveSessionMemberHelperParams.OpParams = &OpRemoveParams;
 	RemoveSessionMemberHelperParams.OpParams->SessionName = TEXT("SessionNameValidRemoveMember");
 
-	FTestPipeline& LoginPipeline = GetLoginPipeline(AccountId);
+	FTestPipeline& LoginPipeline = GetLoginPipeline(UserNumToLogin, { AccountId });
 
 	CreateSessionHelperParams.OpParams->LocalAccountId = AccountId;
 	AddSessionMemberHelperParams.OpParams->LocalAccountId = AccountId;
 	RemoveSessionMemberHelperParams.OpParams->LocalAccountId = AccountId;
+	LeaveSessionHelperParams.OpParams->LocalAccountId = AccountId;
 
 	LoginPipeline
 		.EmplaceStep<FCreateSessionHelper>(MoveTemp(CreateSessionHelperParams))
 		.EmplaceStep<FAddSessionMemberHelper>(MoveTemp(AddSessionMemberHelperParams))
-		.EmplaceStep<FRemoveSessionMemberHelper>(MoveTemp(RemoveSessionMemberHelperParams));
+		.EmplaceStep<FRemoveSessionMemberHelper>(MoveTemp(RemoveSessionMemberHelperParams))
+		.EmplaceStep<FLeaveSessionHelper>(MoveTemp(LeaveSessionHelperParams));
 
 	RunToCompletion();
 }
\ No newline at end of file
diff --git a/Engine/Source/Programs/Online/OnlineTestsCore/Source/Private/Tests/Sessions/CreateSessionTests.cpp b/Engine/Source/Programs/Online/OnlineTestsCore/Source/Private/Tests/Sessions/CreateSessionTests.cpp
--- a/Engine/Source/Programs/Online/OnlineTestsCore/Source/Private/Tests/Sessions/CreateSessionTests.cpp
+++ b/Engine/Source/Programs/Online/OnlineTestsCore/Source/Private/Tests/Sessions/CreateSessionTests.cpp
@@ -2,7 +2,8 @@
 
 #include "Helpers/Sessions/CreateSessionHelper.h"
 #include "Helpers/Sessions/LeaveSessionHelper.h"
-#include "OnlineCatchHelper.h"
+#include "Logging/LogScopedVerbosityOverride.h"
+#include "Online/OnlineServicesLog.h"
 
 #define SESSIONS_TAG "[suite_sessions]"
 #define EG_SESSIONS_CREATESESSION_TAG SESSIONS_TAG "[createsession]"
@@ -12,16 +13,14 @@
 
 SESSIONS_TEST_CASE("If I call CreateSession with an invalid account id, I get an error", EG_SESSIONS_CREATESESSION_TAG)
 {
-	const int32 NumUsersToLogin = 0;
-
 	FCreateSession::Params OpCreateParams; 
 	FCreateSessionHelper::FHelperParams CreateSessionHelperParams;
 	CreateSessionHelperParams.OpParams = &OpCreateParams;
 	CreateSessionHelperParams.OpParams->LocalAccountId = FAccountId();
 	CreateSessionHelperParams.ExpectedError = TOnlineResult<FCreateSession>(Errors::InvalidParams());
 
-	GetLoginPipeline(NumUsersToLogin)
-	.EmplaceStep<FCreateSessionHelper>(MoveTemp(CreateSessionHelperParams));
+	GetPipeline()
+		.EmplaceStep<FCreateSessionHelper>(MoveTemp(CreateSessionHelperParams));
 
 	RunToCompletion();
 }
@@ -36,7 +35,7 @@ SESSIONS_TEST_CASE("If I call CreateSession with an empty session name, I get an
 	CreateSessionHelperParams.OpParams->SessionName = TEXT("");
 	CreateSessionHelperParams.ExpectedError = TOnlineResult<FCreateSession>(Errors::InvalidParams());
 
-	FTestPipeline& LoginPipeline = GetLoginPipeline(AccountId);
+	FTestPipeline& LoginPipeline = GetLoginPipeline({ AccountId });
 
 	CreateSessionHelperParams.OpParams->LocalAccountId = AccountId;
 
@@ -56,7 +55,7 @@ SESSIONS_TEST_CASE("If I call CreateSession with an empty schema name in setting
 	CreateSessionHelperParams.OpParams->SessionSettings.SchemaName = TEXT("");
 	CreateSessionHelperParams.ExpectedError = TOnlineResult<FCreateSession>(Errors::InvalidParams());
 
-	FTestPipeline& LoginPipeline = GetLoginPipeline(AccountId);
+	FTestPipeline& LoginPipeline = GetLoginPipeline({ AccountId });
 
 	CreateSessionHelperParams.OpParams->LocalAccountId = AccountId;
 
@@ -77,7 +76,7 @@ SESSIONS_TEST_CASE("If I call CreateSession with an invalid max connections numb
 	CreateSessionHelperParams.OpParams->SessionSettings.NumMaxConnections = 0;
 	CreateSessionHelperParams.ExpectedError = TOnlineResult<FCreateSession>(Errors::InvalidParams());
 
-	FTestPipeline& LoginPipeline = GetLoginPipeline(AccountId);
+	FTestPipeline& LoginPipeline = GetLoginPipeline({ AccountId });
 
 	CreateSessionHelperParams.OpParams->LocalAccountId = AccountId;
 
@@ -99,7 +98,7 @@ SESSIONS_TEST_CASE("If I call CreateSession with an empty custom setting name in
 	CreateSessionHelperParams.OpParams->SessionSettings.CustomSettings.Emplace(FName(), FCustomSessionSetting{ FSchemaVariant(false), ESchemaAttributeVisibility::Public });
 	CreateSessionHelperParams.ExpectedError = TOnlineResult<FCreateSession>(Errors::InvalidParams());
 
-	FTestPipeline& LoginPipeline = GetLoginPipeline(AccountId);
+	FTestPipeline& LoginPipeline = GetLoginPipeline({ AccountId });
 
 	CreateSessionHelperParams.OpParams->LocalAccountId = AccountId;
 
@@ -110,8 +109,9 @@ SESSIONS_TEST_CASE("If I call CreateSession with an empty custom setting name in
 
 SESSIONS_TEST_CASE("If I call CreateSession with session override id less than 16 chars or more than 64, I get an error", EG_SESSIONS_CREATESESSIONEOS_TAG)
 {
-	ResetAccountStatus();
+	LOG_SCOPE_VERBOSITY_OVERRIDE(LogOnlineServices, ELogVerbosity::NoLogging);
 
+	int32 UserNumToLogin = 7;
 	FAccountId AccountId;
 
 	FCreateSession::Params OpFirstCreateParams;
@@ -134,7 +134,7 @@ SESSIONS_TEST_CASE("If I call CreateSession with session override id less than 1
 	SecondCreateSessionHelperParams.OpParams->SessionSettings.NumMaxConnections = 2;
 	SecondCreateSessionHelperParams.ExpectedError = TOnlineResult<FCreateSession>(Errors::InvalidParams());
 
-	FTestPipeline& LoginPipeline = GetLoginPipeline(AccountId);
+	FTestPipeline& LoginPipeline = GetLoginPipeline(UserNumToLogin, { AccountId });
 
 	FirstCreateSessionHelperParams.OpParams->LocalAccountId = AccountId;
 	SecondCreateSessionHelperParams.OpParams->LocalAccountId = AccountId;
@@ -149,8 +149,8 @@ SESSIONS_TEST_CASE("If I call CreateSession with session override id less than 1
 SESSIONS_TEST_CASE("If I call CreateSession with an name already in use, I get an error", EG_SESSIONS_CREATESESSION_TAG)
 {
 	DestroyCurrentServiceModule();
-	ResetAccountStatus();
-
+	
+	int32 UserNumToLogin = 7;
 	FAccountId AccountId;
 
 	FCreateSession::Params OpCreateParams;
@@ -166,22 +166,34 @@ SESSIONS_TEST_CASE("If I call CreateSession with an name already in use, I get a
 	SecondCreateSessionHelperParams.OpParams = &OpCreateParams;
 	SecondCreateSessionHelperParams.ExpectedError = TOnlineResult<FCreateSession>(Errors::InvalidState());
 
-	FTestPipeline& LoginPipeline = GetLoginPipeline(AccountId);
+	FLeaveSession::Params OpLeaveParams;
+	FLeaveSessionHelper::FHelperParams LeaveSessionHelperParams;
+	LeaveSessionHelperParams.OpParams = &OpLeaveParams;
+	LeaveSessionHelperParams.OpParams->SessionName = TEXT("SessionNameInUse");
+	LeaveSessionHelperParams.OpParams->bDestroySession = true;
+
+	FTestPipeline& LoginPipeline = GetLoginPipeline(UserNumToLogin, { AccountId });
 
 	FirstCreateSessionHelperParams.OpParams->LocalAccountId = AccountId;
 	SecondCreateSessionHelperParams.OpParams->LocalAccountId = AccountId;
+	LeaveSessionHelperParams.OpParams->LocalAccountId = AccountId;
 
 	LoginPipeline
 		.EmplaceStep<FCreateSessionHelper>(MoveTemp(FirstCreateSessionHelperParams))
-		.EmplaceStep<FCreateSessionHelper>(MoveTemp(SecondCreateSessionHelperParams));
+		.EmplaceStep<FCreateSessionHelper>(MoveTemp(SecondCreateSessionHelperParams))
+		.EmplaceStep<FLeaveSessionHelper>(MoveTemp(LeaveSessionHelperParams));
+
 
 	RunToCompletion();
 }
 
 SESSIONS_TEST_CASE("If I call CreateSession with the presence flag set after there is already a presence session, I get an error", EG_SESSIONS_CREATESESSION_TAG)
 {
-	DestroyCurrentServiceModule();
+	LOG_SCOPE_VERBOSITY_OVERRIDE(LogOnlineServices, ELogVerbosity::NoLogging);
 
+	DestroyCurrentServiceModule();
+	
+	int32 UserNumToLogin = 7;
 	FAccountId AccountId;
 
 	FCreateSession::Params OpFirstCreateParams;
@@ -207,7 +219,7 @@ SESSIONS_TEST_CASE("If I call CreateSession with the presence flag set after the
 	LeaveSessionHelperParams.OpParams->SessionName = TEXT("SessionNamePresenceEnabled1");
 	LeaveSessionHelperParams.OpParams->bDestroySession = true;
 
-	FTestPipeline& LoginPipeline = GetLoginPipeline(AccountId);
+	FTestPipeline& LoginPipeline = GetLoginPipeline(UserNumToLogin, { AccountId });
 
 	FirstCreateSessionHelperParams.OpParams->LocalAccountId = AccountId;
 	SecondCreateSessionHelperParams.OpParams->LocalAccountId = AccountId;
@@ -235,12 +247,20 @@ SESSIONS_TEST_CASE("If I call CreateSession with valid data, the operation compl
 	CreateSessionHelperParams.OpParams->SessionSettings.NumMaxConnections = 2;
 	CreateSessionHelperParams.OpParams->bPresenceEnabled = true;
 	
-	FTestPipeline& LoginPipeline = GetLoginPipeline(AccountId);
+	FLeaveSession::Params OpLeaveParams;
+	FLeaveSessionHelper::FHelperParams LeaveSessionHelperParams;
+	LeaveSessionHelperParams.OpParams = &OpLeaveParams;
+	LeaveSessionHelperParams.OpParams->SessionName = TEXT("SessionNameValid");
+	LeaveSessionHelperParams.OpParams->bDestroySession = true;
+
+	FTestPipeline& LoginPipeline = GetLoginPipeline({ AccountId });
 
 	CreateSessionHelperParams.OpParams->LocalAccountId = AccountId;
+	LeaveSessionHelperParams.OpParams->LocalAccountId = AccountId;
 
 	LoginPipeline
-		.EmplaceStep<FCreateSessionHelper>(MoveTemp(CreateSessionHelperParams));
+		.EmplaceStep<FCreateSessionHelper>(MoveTemp(CreateSessionHelperParams))
+		.EmplaceStep<FLeaveSessionHelper>(MoveTemp(LeaveSessionHelperParams));
 
 	RunToCompletion();
 }
@@ -268,14 +288,22 @@ SESSIONS_TEST_CASE("If I call CreateSession twice for NULL, I get an error", EG_
 	SecondCreateSessionHelperParams.OpParams->SessionSettings.NumMaxConnections = 4;
 	SecondCreateSessionHelperParams.ExpectedError = TOnlineResult<FCreateSession>(Errors::AlreadyPending());
 
-	FTestPipeline& LoginPipeline = GetLoginPipeline(FirstAccountId, SecondAccountId);
+	FLeaveSession::Params OpLeaveParams;
+	FLeaveSessionHelper::FHelperParams LeaveSessionHelperParams;
+	LeaveSessionHelperParams.OpParams = &OpLeaveParams;
+	LeaveSessionHelperParams.OpParams->SessionName = TEXT("CreateSessionTwiceName");
+	LeaveSessionHelperParams.OpParams->bDestroySession = true;
+
+	FTestPipeline& LoginPipeline = GetLoginPipeline({ FirstAccountId, SecondAccountId });
 
 	FirstCreateSessionHelperParams.OpParams->LocalAccountId = FirstAccountId;
 	SecondCreateSessionHelperParams.OpParams->LocalAccountId = FirstAccountId;
+	LeaveSessionHelperParams.OpParams->LocalAccountId = FirstAccountId;
 
 	LoginPipeline
 		.EmplaceStep<FCreateSessionHelper>(MoveTemp(FirstCreateSessionHelperParams))
-		.EmplaceStep<FCreateSessionHelper>(MoveTemp(SecondCreateSessionHelperParams));
+		.EmplaceStep<FCreateSessionHelper>(MoveTemp(SecondCreateSessionHelperParams))
+		.EmplaceStep<FLeaveSessionHelper>(MoveTemp(LeaveSessionHelperParams));
 
 	RunToCompletion();
 }
\ No newline at end of file
diff --git a/Engine/Source/Programs/Online/OnlineTestsCore/Source/Private/Tests/Sessions/FindSessionsTests.cpp b/Engine/Source/Programs/Online/OnlineTestsCore/Source/Private/Tests/Sessions/FindSessionsTests.cpp
--- a/Engine/Source/Programs/Online/OnlineTestsCore/Source/Private/Tests/Sessions/FindSessionsTests.cpp
+++ b/Engine/Source/Programs/Online/OnlineTestsCore/Source/Private/Tests/Sessions/FindSessionsTests.cpp
@@ -1,11 +1,14 @@
 // Copyright Epic Games, Inc. All Rights Reserved.
 
+
 #include "Helpers/Sessions/CreateSessionHelper.h"
 #include "Helpers/Sessions/AddRemoveSessionMemberHelper.h"
 #include "Helpers/Sessions/FindSessionsHelper.h"
 #include "Helpers/Sessions/LeaveSessionHelper.h"
 #include "Helpers/TickForTime.h"
-#include "OnlineCatchHelper.h"
+#include "Logging/LogScopedVerbosityOverride.h"
+#include "Online/OnlineServicesLog.h"
+#include "EOSShared.h"
 
 #define SESSIONS_TAG "[suite_sessions]"
 #define EG_SESSIONS_FINDSESSIONS_TAG SESSIONS_TAG "[findsessions]"
@@ -14,15 +17,13 @@
 
 SESSIONS_TEST_CASE("If I call FindSessions with an invalid account id, I get an error", EG_SESSIONS_FINDSESSIONS_TAG)
 {
-	const int32 NumUsersToLogin = 0;
-
 	FFindSessions::Params OpFindParams;
 	FFindSessionsHelper::FHelperParams FindSessionsHelperParams;
 	FindSessionsHelperParams.OpParams = &OpFindParams;
 	FindSessionsHelperParams.OpParams->LocalAccountId = FAccountId();
 	FindSessionsHelperParams.ExpectedError = TOnlineResult<FFindSessions>(Errors::InvalidParams());
 
-	GetLoginPipeline(NumUsersToLogin)
+	GetPipeline()
 		.EmplaceStep<FFindSessionsHelper>(MoveTemp(FindSessionsHelperParams));
 
 	RunToCompletion();
@@ -38,7 +39,7 @@ SESSIONS_TEST_CASE("If I call FindSessions with an invalid max results number, I
 	FindSessionsHelperParams.OpParams->MaxResults = 0;
 	FindSessionsHelperParams.ExpectedError = TOnlineResult<FFindSessions>(Errors::InvalidParams());
 
-	FTestPipeline& LoginPipeline = GetLoginPipeline(AccountId);
+	FTestPipeline& LoginPipeline = GetLoginPipeline({ AccountId });
 
 	FindSessionsHelperParams.OpParams->LocalAccountId = AccountId;
 
@@ -59,7 +60,7 @@ SESSIONS_TEST_CASE("If I call FindSessions with an empty custom setting name in
 	FindSessionsHelperParams.OpParams->Filters.Add(FFindSessionsSearchFilter{ FName(), ESchemaAttributeComparisonOp::Equals, FSchemaVariant(false) });
 	FindSessionsHelperParams.ExpectedError = TOnlineResult<FFindSessions>(Errors::InvalidParams());
 
-	FTestPipeline& LoginPipeline = GetLoginPipeline(AccountId);
+	FTestPipeline& LoginPipeline = GetLoginPipeline({ AccountId });
 
 	FindSessionsHelperParams.OpParams->LocalAccountId = AccountId;
 
@@ -80,7 +81,7 @@ SESSIONS_TEST_CASE("If I call FindSessions with an invalid session id as filter,
 	FindSessionsHelperParams.OpParams->SessionId = FOnlineSessionId();
 	FindSessionsHelperParams.ExpectedError = TOnlineResult<FFindSessions>(Errors::InvalidParams());
 
-	FTestPipeline& LoginPipeline = GetLoginPipeline(AccountId);
+	FTestPipeline& LoginPipeline = GetLoginPipeline({ AccountId });
 	
 	FindSessionsHelperParams.OpParams->LocalAccountId = AccountId;
 
@@ -101,7 +102,7 @@ SESSIONS_TEST_CASE("If I call FindSessions with an invalid account id as filter,
 	FindSessionsHelperParams.OpParams->TargetUser = FAccountId();
 	FindSessionsHelperParams.ExpectedError = TOnlineResult<FFindSessions>(Errors::InvalidParams());
 
-	FTestPipeline& LoginPipeline = GetLoginPipeline(AccountId);
+	FTestPipeline& LoginPipeline = GetLoginPipeline({ AccountId });
 	
 	FindSessionsHelperParams.OpParams->LocalAccountId = AccountId;
 
@@ -122,8 +123,7 @@ SESSIONS_TEST_CASE("If I call FindSessions with an empty search filter key, I ge
 	FindSessionsHelperParams.OpParams->Filters.Emplace(FFindSessionsSearchFilter{ FName(""), ESchemaAttributeComparisonOp::Equals, FSchemaVariant(int64(64)) });
 	FindSessionsHelperParams.ExpectedError = TOnlineResult<FFindSessions>(Errors::InvalidParams());
 
-
-	FTestPipeline& LoginPipeline = GetLoginPipeline(AccountId);
+	FTestPipeline& LoginPipeline = GetLoginPipeline({ AccountId });
 
 	FindSessionsHelperParams.OpParams->LocalAccountId = AccountId;
 
@@ -136,8 +136,9 @@ SESSIONS_TEST_CASE("If I call FindSessions with an empty search filter key, I ge
 SESSIONS_TEST_CASE("If I call FindSessions while a search is already running for that user, I get an error", EG_SESSIONS_FINDSESSIONS_TAG)
 {
 	DestroyCurrentServiceModule();
-	ResetAccountStatus();
+	ReturnAccounts();
 
+	int32 UserNumToLogin = 7;
 	FAccountId FirstAccountId, SecondAccountId;
 
 	FCreateSession::Params OpCreateParams;
@@ -170,30 +171,29 @@ SESSIONS_TEST_CASE("If I call FindSessions while a search is already running for
 	LeaveSessionHelperParams.OpParams->SessionName = TEXT("AlreadyPendingName");
 	LeaveSessionHelperParams.OpParams->bDestroySession = true;
 
-	FTestPipeline& LoginPipeline = GetLoginPipeline(FirstAccountId, SecondAccountId);
+	FTestPipeline& LoginPipeline = GetLoginPipeline(UserNumToLogin, { FirstAccountId, SecondAccountId });
 
 	CreateSessionHelperParams.OpParams->LocalAccountId = FirstAccountId;
 
 	AddSessionMemberHelperParams.OpParams->LocalAccountId = FirstAccountId;
 
-	LeaveSessionHelperParams.OpParams->LocalAccountId = FirstAccountId;
-
 	FindFirstSessionsHelperParams.OpParams->LocalAccountId = SecondAccountId;
 	FindFirstSessionsHelperParams.OpParams->TargetUser = FirstAccountId;
 
 	FindSecondSessionsHelperParams.OpParams->LocalAccountId = SecondAccountId;
 	FindSecondSessionsHelperParams.OpParams->TargetUser = FirstAccountId;
 
+	LeaveSessionHelperParams.OpParams->LocalAccountId = FirstAccountId;
+
 	LoginPipeline
 		.EmplaceStep<FCreateSessionHelper>(MoveTemp(CreateSessionHelperParams))
 		.EmplaceStep<FAddSessionMemberHelper>(MoveTemp(AddSessionMemberHelperParams))
 		.EmplaceAsyncLambda([&FindFirstSessionsHelperParams, &FindSecondSessionsHelperParams](FAsyncLambdaResult Promise, SubsystemType OnlineSubsystem)
 			{
 				ISessionsPtr SessionsInterface = OnlineSubsystem->GetSessionsInterface();
-				EOnlineServices CurrentService = OnlineSubsystem->GetServicesProvider();
 
 				SessionsInterface->FindSessions(MoveTemp(*FindFirstSessionsHelperParams.OpParams))
-					.OnComplete([Promise = MoveTemp(Promise), CurrentService](const TOnlineResult<FFindSessions>& Result)
+					.OnComplete([Promise = MoveTemp(Promise)](const TOnlineResult<FFindSessions>& Result)
 						{
 							CHECK_OP(Result);
 							Promise->SetValue(true);
@@ -222,7 +222,7 @@ SESSIONS_TEST_CASE("If I call FindSessions with valid information, I can get an
 	FindSessionsHelperParams.OpParams = &OpFirstParams;
 	FindSessionsHelperParams.OpParams->MaxResults = 3;
 
-	FTestPipeline& LoginPipeline = GetLoginPipeline(AccountId);
+	FTestPipeline& LoginPipeline = GetLoginPipeline({ AccountId });
 
 	FindSessionsHelperParams.OpParams->LocalAccountId = AccountId;
 	FindSessionsHelperParams.OpParams->Filters.Emplace(FFindSessionsSearchFilter{ FName("NonExistentSettingName"), ESchemaAttributeComparisonOp::Equals, FSchemaVariant(int64(64)) });
@@ -236,7 +236,7 @@ SESSIONS_TEST_CASE("If I call FindSessions with valid information, I can get an
 SESSIONS_TEST_CASE("If I call FindSessions with valid information, I will still not find sessions which JoinPolicy is not public", EG_SESSIONS_FINDSESSIONS_TAG)
 {
 	DestroyCurrentServiceModule();
-	ResetAccountStatus();
+	ReturnAccounts();
 	
 	FAccountId FirstAccountId, SecondAccountId;
 	
@@ -265,7 +265,7 @@ SESSIONS_TEST_CASE("If I call FindSessions with valid information, I will still
 	LeaveSessionHelperParams.OpParams->SessionName = TEXT("FriendsOnlyName");
 	LeaveSessionHelperParams.OpParams->bDestroySession = true;
 	
-	FTestPipeline& LoginPipeline = GetLoginPipeline(FirstAccountId, SecondAccountId);
+	FTestPipeline& LoginPipeline = GetLoginPipeline({ FirstAccountId, SecondAccountId });
 	
 	CreateSessionHelperParams.OpParams->LocalAccountId = FirstAccountId;
 	
@@ -287,8 +287,9 @@ SESSIONS_TEST_CASE("If I call FindSessions with valid information, I will still
 
 SESSIONS_TEST_CASE("If I call FindSessions with valid information, I will get a list of ids for all the valid sessions", EG_SESSIONS_FINDSESSIONSEOS_TAG)
 {
+	LOG_SCOPE_VERBOSITY_OVERRIDE(LogEOSSDK, ELogVerbosity::NoLogging);
+
 	DestroyCurrentServiceModule();
-	ResetAccountStatus();
 
 	FAccountId FirstAccountId, SecondAccountId;
 
@@ -335,7 +336,7 @@ SESSIONS_TEST_CASE("If I call FindSessions with valid information, I will get a
 	SecondLeaveSessionHelperParams.OpParams->SessionName = TEXT("ValidName1");
 	SecondLeaveSessionHelperParams.OpParams->bDestroySession = true;
 
-	FTestPipeline& LoginPipeline = GetLoginPipeline(FirstAccountId, SecondAccountId);
+	FTestPipeline& LoginPipeline = GetLoginPipeline({ FirstAccountId, SecondAccountId });
 
 	FirstCreateSessionHelperParams.OpParams->LocalAccountId = FirstAccountId;
 	SecondCreateSessionHelperParams.OpParams->LocalAccountId = FirstAccountId;
@@ -366,6 +367,8 @@ SESSIONS_TEST_CASE("If I call FindSessions with valid information, I will get a
 
 SESSIONS_TEST_CASE("Call FindSessions using custom setting as a filter Bool", EG_SESSIONS_FINDSESSIONSEOS_TAG)
 {
+	LOG_SCOPE_VERBOSITY_OVERRIDE(LogEOSSDK, ELogVerbosity::NoLogging);
+
 	DestroyCurrentServiceModule();
 
 	FAccountId FirstAccountId, SecondAccountId;
@@ -399,13 +402,19 @@ SESSIONS_TEST_CASE("Call FindSessions using custom setting as a filter Bool", EG
 	SecondFindSessionsHelperParams.OpParams->MaxResults = 1;
 	SecondFindSessionsHelperParams.OpParams->Filters.Emplace(FFindSessionsSearchFilter{ FName("BoolSettingName"), ESchemaAttributeComparisonOp::NotEquals, FSchemaVariant(true) });
 
-	FLeaveSession::Params OpLeaveParams;
-	FLeaveSessionHelper::FHelperParams LeaveSessionHelperParams;
-	LeaveSessionHelperParams.OpParams = &OpLeaveParams;
-	LeaveSessionHelperParams.OpParams->SessionName = TEXT("FindSessionNameBoolFilter");
-	LeaveSessionHelperParams.OpParams->bDestroySession = true;
+	FLeaveSession::Params OpFirstLeaveParams;
+	FLeaveSessionHelper::FHelperParams FirstLeaveSessionHelperParams;
+	FirstLeaveSessionHelperParams.OpParams = &OpFirstLeaveParams;
+	FirstLeaveSessionHelperParams.OpParams->SessionName = TEXT("FindSessionNameBoolFilter");
+	FirstLeaveSessionHelperParams.OpParams->bDestroySession = true;
+
+	FLeaveSession::Params OpSecondLeaveParams;
+	FLeaveSessionHelper::FHelperParams SecondLeaveSessionHelperParams;
+	SecondLeaveSessionHelperParams.OpParams = &OpSecondLeaveParams;
+	SecondLeaveSessionHelperParams.OpParams->SessionName = TEXT("FindSessionNameBoolFilter1");
+	SecondLeaveSessionHelperParams.OpParams->bDestroySession = true;
 
-	FTestPipeline& LoginPipeline = GetLoginPipeline(FirstAccountId, SecondAccountId);
+	FTestPipeline& LoginPipeline = GetLoginPipeline({ FirstAccountId, SecondAccountId });
 
 	FirstCreateSessionHelperParams.OpParams->LocalAccountId = FirstAccountId;
 	SecondCreateSessionHelperParams.OpParams->LocalAccountId = FirstAccountId;
@@ -413,7 +422,8 @@ SESSIONS_TEST_CASE("Call FindSessions using custom setting as a filter Bool", EG
 	FirstFindSessionsHelperParams.OpParams->LocalAccountId = FirstAccountId;
 	SecondFindSessionsHelperParams.OpParams->LocalAccountId = SecondAccountId;
 
-	LeaveSessionHelperParams.OpParams->LocalAccountId = FirstAccountId;
+	FirstLeaveSessionHelperParams.OpParams->LocalAccountId = FirstAccountId;
+	SecondLeaveSessionHelperParams.OpParams->LocalAccountId = FirstAccountId;
 
 	uint32_t ExpectedSessionsFound = 1;
 
@@ -423,13 +433,16 @@ SESSIONS_TEST_CASE("Call FindSessions using custom setting as a filter Bool", EG
 	.EmplaceStep<FTickForTime>(FTimespan::FromMilliseconds(6000))
 	.EmplaceStep<FFindSessionsHelper>(MoveTemp(FirstFindSessionsHelperParams), ExpectedSessionsFound)
 	.EmplaceStep<FFindSessionsHelper>(MoveTemp(SecondFindSessionsHelperParams))
-	.EmplaceStep<FLeaveSessionHelper>(MoveTemp(LeaveSessionHelperParams));
+	.EmplaceStep<FLeaveSessionHelper>(MoveTemp(FirstLeaveSessionHelperParams))
+	.EmplaceStep<FLeaveSessionHelper>(MoveTemp(SecondLeaveSessionHelperParams));
 
 	RunToCompletion();
 }
 
 SESSIONS_TEST_CASE("Call FindSessions using custom setting as a filter Int", EG_SESSIONS_FINDSESSIONSEOS_TAG)
 {
+	LOG_SCOPE_VERBOSITY_OVERRIDE(LogEOSSDK, ELogVerbosity::NoLogging);
+
 	DestroyCurrentServiceModule();
 
 	FAccountId FirstAccountId, SecondAccountId;
@@ -487,17 +500,25 @@ SESSIONS_TEST_CASE("Call FindSessions using custom setting as a filter Int", EG_
 	EighthFindSessionsHelperParams.OpParams->MaxResults = 1;
 	EighthFindSessionsHelperParams.OpParams->Filters.Emplace(FFindSessionsSearchFilter{ FName("Int64SettingName"), ESchemaAttributeComparisonOp::LessThanEquals, FSchemaVariant(int64(66)) });
 
-	FLeaveSession::Params OpLeaveParams;
-	FLeaveSessionHelper::FHelperParams LeaveSessionHelperParams;
-	LeaveSessionHelperParams.OpParams = &OpLeaveParams;
-	LeaveSessionHelperParams.OpParams->SessionName = TEXT("FindSessionNameIntFilter");
-	LeaveSessionHelperParams.OpParams->bDestroySession = true;
+	FLeaveSession::Params OpFirstLeaveParams;
+	FLeaveSessionHelper::FHelperParams FirstLeaveSessionHelperParams;
+	FirstLeaveSessionHelperParams.OpParams = &OpFirstLeaveParams;
+	FirstLeaveSessionHelperParams.OpParams->SessionName = TEXT("FindSessionNameIntFilter");
+	FirstLeaveSessionHelperParams.OpParams->bDestroySession = true;
 
-	FTestPipeline& LoginPipeline = GetLoginPipeline(FirstAccountId, SecondAccountId);
+	FLeaveSession::Params OpSecondLeaveParams;
+	FLeaveSessionHelper::FHelperParams SecondLeaveSessionHelperParams;
+	SecondLeaveSessionHelperParams.OpParams = &OpSecondLeaveParams;
+	SecondLeaveSessionHelperParams.OpParams->SessionName = TEXT("FindSessionNameIntFilter1");
+	SecondLeaveSessionHelperParams.OpParams->bDestroySession = true;
+
+	FTestPipeline& LoginPipeline = GetLoginPipeline({ FirstAccountId, SecondAccountId });
 
 	FirstCreateSessionHelperParams.OpParams->LocalAccountId = FirstAccountId;
 	SecondCreateSessionHelperParams.OpParams->LocalAccountId = FirstAccountId;
-	LeaveSessionHelperParams.OpParams->LocalAccountId = FirstAccountId;
+
+	FirstLeaveSessionHelperParams.OpParams->LocalAccountId = FirstAccountId;
+	SecondLeaveSessionHelperParams.OpParams->LocalAccountId = FirstAccountId;
 
 	ThirdFindSessionsHelperParams.OpParams->LocalAccountId = SecondAccountId;
 	FourthFindSessionsHelperParams.OpParams->LocalAccountId = FirstAccountId;
@@ -518,13 +539,16 @@ SESSIONS_TEST_CASE("Call FindSessions using custom setting as a filter Int", EG_
 		.EmplaceStep<FFindSessionsHelper>(MoveTemp(SixthFindSessionsHelperParams), ExpectedSessionsFound)
 		.EmplaceStep<FFindSessionsHelper>(MoveTemp(SeventhFindSessionsHelperParams), ExpectedSessionsFound)
 		.EmplaceStep<FFindSessionsHelper>(MoveTemp(EighthFindSessionsHelperParams), ExpectedSessionsFound)
-		.EmplaceStep<FLeaveSessionHelper>(MoveTemp(LeaveSessionHelperParams));
+		.EmplaceStep<FLeaveSessionHelper>(MoveTemp(FirstLeaveSessionHelperParams))
+		.EmplaceStep<FLeaveSessionHelper>(MoveTemp(SecondLeaveSessionHelperParams));
 
 	RunToCompletion();
 }
 
 SESSIONS_TEST_CASE("Call FindSessions using custom setting as a filter Double", EG_SESSIONS_FINDSESSIONSEOS_TAG)
 {
+	LOG_SCOPE_VERBOSITY_OVERRIDE(LogEOSSDK, ELogVerbosity::NoLogging);
+
 	DestroyCurrentServiceModule();
 
 	FAccountId FirstAccountId, SecondAccountId;
@@ -590,17 +614,25 @@ SESSIONS_TEST_CASE("Call FindSessions using custom setting as a filter Double",
 	FifteenthFindSessionsHelperParams.OpParams->MaxResults = 1;
 	FifteenthFindSessionsHelperParams.OpParams->Filters.Emplace(FFindSessionsSearchFilter{ FName("DoubleSettingName"), ESchemaAttributeComparisonOp::Near, FSchemaVariant(90.0) });
 
-	FLeaveSession::Params OpLeaveParams;
-	FLeaveSessionHelper::FHelperParams LeaveSessionHelperParams;
-	LeaveSessionHelperParams.OpParams = &OpLeaveParams;
-	LeaveSessionHelperParams.OpParams->SessionName = TEXT("FindSessionNameDoubleFilter");
-	LeaveSessionHelperParams.OpParams->bDestroySession = true;
+	FLeaveSession::Params OpFirstLeaveParams;
+	FLeaveSessionHelper::FHelperParams FirstLeaveSessionHelperParams;
+	FirstLeaveSessionHelperParams.OpParams = &OpFirstLeaveParams;
+	FirstLeaveSessionHelperParams.OpParams->SessionName = TEXT("FindSessionNameDoubleFilter");
+	FirstLeaveSessionHelperParams.OpParams->bDestroySession = true;
+
+	FLeaveSession::Params OpSecondLeaveParams;
+	FLeaveSessionHelper::FHelperParams SecondLeaveSessionHelperParams;
+	SecondLeaveSessionHelperParams.OpParams = &OpSecondLeaveParams;
+	SecondLeaveSessionHelperParams.OpParams->SessionName = TEXT("FindSessionNameDoubleFilter1");
+	SecondLeaveSessionHelperParams.OpParams->bDestroySession = true;
 
-	FTestPipeline& LoginPipeline = GetLoginPipeline(FirstAccountId, SecondAccountId);
+	FTestPipeline& LoginPipeline = GetLoginPipeline({ FirstAccountId, SecondAccountId });
 
 	FirstCreateSessionHelperParams.OpParams->LocalAccountId = FirstAccountId;
 	SecondCreateSessionHelperParams.OpParams->LocalAccountId = FirstAccountId;
-	LeaveSessionHelperParams.OpParams->LocalAccountId = FirstAccountId;
+
+	FirstLeaveSessionHelperParams.OpParams->LocalAccountId = FirstAccountId;
+	SecondLeaveSessionHelperParams.OpParams->LocalAccountId = FirstAccountId;
 
 	NinthFindSessionsHelperParams.OpParams->LocalAccountId = SecondAccountId;
 	TenthFindSessionsHelperParams.OpParams->LocalAccountId = FirstAccountId;
@@ -644,13 +676,17 @@ SESSIONS_TEST_CASE("Call FindSessions using custom setting as a filter Double",
 						});
 
 			})
-		.EmplaceStep<FLeaveSessionHelper>(MoveTemp(LeaveSessionHelperParams));
+		.EmplaceStep<FLeaveSessionHelper>(MoveTemp(FirstLeaveSessionHelperParams))
+		.EmplaceStep<FLeaveSessionHelper>(MoveTemp(SecondLeaveSessionHelperParams));
 
 	RunToCompletion();
 }
 
 SESSIONS_TEST_CASE("Call FindSessions using custom setting as a filter String", EG_SESSIONS_FINDSESSIONSEOS_TAG)
 {
+	LOG_SCOPE_VERBOSITY_OVERRIDE(LogOnlineServices, ELogVerbosity::NoLogging);
+	LOG_SCOPE_VERBOSITY_OVERRIDE(LogEOSSDK, ELogVerbosity::NoLogging);
+
 	DestroyCurrentServiceModule();
 
 	FAccountId FirstAccountId, SecondAccountId;
@@ -697,17 +733,25 @@ SESSIONS_TEST_CASE("Call FindSessions using custom setting as a filter String",
 	NineteenthFindSessionsHelperParams.OpParams->Filters.Emplace(FFindSessionsSearchFilter{ FName("StringSettingName"), ESchemaAttributeComparisonOp::In, FSchemaVariant(TEXT("TestStringData;OtherStringData;AnotherStringData")) });
 	NineteenthFindSessionsHelperParams.ExpectedError = TOnlineResult<FFindSessions>(Errors::InvalidParams());
 	
-	FLeaveSession::Params OpLeaveParams;
-	FLeaveSessionHelper::FHelperParams LeaveSessionHelperParams;
-	LeaveSessionHelperParams.OpParams = &OpLeaveParams;
-	LeaveSessionHelperParams.OpParams->SessionName = TEXT("FindSessionNameStringFilter");
-	LeaveSessionHelperParams.OpParams->bDestroySession = true;
+	FLeaveSession::Params OpFirstLeaveParams;
+	FLeaveSessionHelper::FHelperParams FirstLeaveSessionHelperParams;
+	FirstLeaveSessionHelperParams.OpParams = &OpFirstLeaveParams;
+	FirstLeaveSessionHelperParams.OpParams->SessionName = TEXT("FindSessionNameStringFilter");
+	FirstLeaveSessionHelperParams.OpParams->bDestroySession = true;
+
+	FLeaveSession::Params OpSecondLeaveParams;
+	FLeaveSessionHelper::FHelperParams SecondLeaveSessionHelperParams;
+	SecondLeaveSessionHelperParams.OpParams = &OpSecondLeaveParams;
+	SecondLeaveSessionHelperParams.OpParams->SessionName = TEXT("FindSessionNameStringFilter1");
+	SecondLeaveSessionHelperParams.OpParams->bDestroySession = true;
 
-	FTestPipeline& LoginPipeline = GetLoginPipeline(FirstAccountId, SecondAccountId);
+	FTestPipeline& LoginPipeline = GetLoginPipeline({ FirstAccountId, SecondAccountId });
 
 	FirstCreateSessionHelperParams.OpParams->LocalAccountId = FirstAccountId;
 	SecondCreateSessionHelperParams.OpParams->LocalAccountId = FirstAccountId;
-	LeaveSessionHelperParams.OpParams->LocalAccountId = FirstAccountId;
+	
+	FirstLeaveSessionHelperParams.OpParams->LocalAccountId = FirstAccountId;
+	SecondLeaveSessionHelperParams.OpParams->LocalAccountId = FirstAccountId;
 
 	SixteenthFindSessionsHelperParams.OpParams->LocalAccountId = SecondAccountId;
 	SeventeenthFindSessionsHelperParams.OpParams->LocalAccountId = FirstAccountId;
@@ -726,9 +770,10 @@ SESSIONS_TEST_CASE("Call FindSessions using custom setting as a filter String",
 		.EmplaceStep<FFindSessionsHelper>(MoveTemp(EighteenthFindSessionsHelperParams), ExpectedSessionsFound)
 		.EmplaceStep<FFindSessionsHelper>(MoveTemp(NineteenthFindSessionsHelperParams), NoSessionsFound, []()
 			{
-				UE_LOG_ONLINETESTS(Warning, TEXT("Functionality being tested is still unsupported as of EOSSDK Version: 1.15.3-21924193"));
+				UE_LOG_ONLINETESTS(Verbose, TEXT("Functionality being tested is still unsupported as of EOSSDK Version: 1.15.3-21924193"));
 			})
-		.EmplaceStep<FLeaveSessionHelper>(MoveTemp(LeaveSessionHelperParams));
+		.EmplaceStep<FLeaveSessionHelper>(MoveTemp(FirstLeaveSessionHelperParams))
+		.EmplaceStep<FLeaveSessionHelper>(MoveTemp(SecondLeaveSessionHelperParams));
 
 	RunToCompletion();
 }
\ No newline at end of file
diff --git a/Engine/Source/Programs/Online/OnlineTestsCore/Source/Private/Tests/Sessions/GetSessionsTests.cpp b/Engine/Source/Programs/Online/OnlineTestsCore/Source/Private/Tests/Sessions/GetSessionsTests.cpp
--- a/Engine/Source/Programs/Online/OnlineTestsCore/Source/Private/Tests/Sessions/GetSessionsTests.cpp
+++ b/Engine/Source/Programs/Online/OnlineTestsCore/Source/Private/Tests/Sessions/GetSessionsTests.cpp
@@ -3,8 +3,10 @@
 #include "Helpers/Sessions/CreateSessionHelper.h"
 #include "Helpers/Sessions/LeaveSessionHelper.h"
 #include "Helpers/Sessions/SendRejectSessionInviteHelper.h"
+#include "Logging/LogScopedVerbosityOverride.h"
+#include "Online/OnlineServicesLog.h"
+#include "EOSShared.h"
 #include "Helpers/TickForTime.h"
-#include "OnlineCatchHelper.h"
 
 #define SESSIONS_TAG "[suite_sessions]"
 #define EG_SESSIONS_GETSESSIONS_TAG SESSIONS_TAG "[getsessions]"
@@ -13,7 +15,7 @@
 
 SESSIONS_TEST_CASE("If I call GetAllSessions with an invalid account id, I get an error", EG_SESSIONS_GETSESSIONS_TAG)
 {
-	GetLoginPipeline()
+	GetPipeline()
 		.EmplaceLambda([](SubsystemType OnlineSubsystem)
 			{
 				FGetAllSessions::Params OpGetAllParams;
@@ -31,9 +33,10 @@ SESSIONS_TEST_CASE("If I call GetAllSessions with an invalid account id, I get a
 SESSIONS_TEST_CASE("If I call GetAllSessions before creating or joining any sessions, I get a successful result which is an empty array", EG_SESSIONS_GETSESSIONS_TAG)
 {
 	DestroyCurrentServiceModule();
+	ReturnAccounts();
 
 	FAccountId AccountId;
-	GetLoginPipeline(AccountId)
+	GetLoginPipeline({ AccountId })
 		.EmplaceLambda([&AccountId](SubsystemType OnlineSubsystem)
 			{
 				FGetAllSessions::Params OpParams;
@@ -50,6 +53,8 @@ SESSIONS_TEST_CASE("If I call GetAllSessions before creating or joining any sess
 
 SESSIONS_TEST_CASE("If I call GetAllSessions with valid conditions, I get a valid array of session references", EG_SESSIONS_GETSESSIONS_TAG)
 {
+	LOG_SCOPE_VERBOSITY_OVERRIDE(LogOnlineServices, ELogVerbosity::NoLogging);
+
 	FAccountId AccountId;
 
 	FCreateSession::Params OpCreateParams;
@@ -60,9 +65,16 @@ SESSIONS_TEST_CASE("If I call GetAllSessions with valid conditions, I get a vali
 	CreateSessionHelperParams.OpParams->SessionSettings.NumMaxConnections = 4;
 	CreateSessionHelperParams.OpParams->bPresenceEnabled = true;
 
-	FTestPipeline& LoginPipeline = GetLoginPipeline(AccountId);
+	FLeaveSession::Params OpLeaveParams;
+	FLeaveSessionHelper::FHelperParams LeaveSessionHelperParams;
+	LeaveSessionHelperParams.OpParams = &OpLeaveParams;
+	LeaveSessionHelperParams.OpParams->SessionName = TEXT("GetAllSessionsValidName");
+	LeaveSessionHelperParams.OpParams->bDestroySession = true;
+
+	FTestPipeline& LoginPipeline = GetLoginPipeline({ AccountId });
 
 	CreateSessionHelperParams.OpParams->LocalAccountId = AccountId;
+	LeaveSessionHelperParams.OpParams->LocalAccountId = AccountId;
 
 	uint32_t ExpectedSessionsFound = 1;
 
@@ -77,7 +89,9 @@ SESSIONS_TEST_CASE("If I call GetAllSessions with valid conditions, I get a vali
 				TOnlineResult<FGetAllSessions> Result = SessionsInterface->GetAllSessions(MoveTemp(OpGetAllParams));
 				REQUIRE_OP(Result);
 				CHECK(Result.GetOkValue().Sessions.Num() == ExpectedSessionsFound);
-			});
+			})
+		.EmplaceStep<FLeaveSessionHelper>(MoveTemp(LeaveSessionHelperParams));
+
 
 	RunToCompletion();
 }
@@ -88,7 +102,7 @@ SESSIONS_TEST_CASE("If I call GetSessionByName with an empty session name, I get
 
 	FAccountId AccountId;
 
-	GetLoginPipeline(AccountId)
+	GetLoginPipeline({ AccountId })
 		.EmplaceLambda([](SubsystemType OnlineSubsystem)
 			{
 				FGetSessionByName::Params OpParams;
@@ -107,7 +121,7 @@ SESSIONS_TEST_CASE("If I call GetSessionByName with an unregistered session name
 {
 	FAccountId AccountId;
 
-	GetLoginPipeline(AccountId)
+	GetLoginPipeline({ AccountId })
 		.EmplaceLambda([](SubsystemType OnlineSubsystem)
 			{
 				FGetSessionByName::Params OpGetByNameParams;
@@ -124,7 +138,9 @@ SESSIONS_TEST_CASE("If I call GetSessionByName with an unregistered session name
 
 SESSIONS_TEST_CASE("If I call GetSessionByName with valid information, it returns a valid session reference", EG_SESSIONS_GETSESSIONS_TAG)
 {
-	ResetAccountStatus();
+	LOG_SCOPE_VERBOSITY_OVERRIDE(LogOnlineServices, ELogVerbosity::NoLogging);
+
+	DestroyCurrentServiceModule();
 
 	FAccountId AccountId;
 
@@ -136,9 +152,16 @@ SESSIONS_TEST_CASE("If I call GetSessionByName with valid information, it return
 	CreateSessionHelperParams.OpParams->SessionSettings.NumMaxConnections = 2;
 	CreateSessionHelperParams.OpParams->bPresenceEnabled = true;
 
-	FTestPipeline& LoginPipeline = GetLoginPipeline(AccountId);
+	FLeaveSession::Params OpLeaveParams;
+	FLeaveSessionHelper::FHelperParams LeaveSessionHelperParams;
+	LeaveSessionHelperParams.OpParams = &OpLeaveParams;
+	LeaveSessionHelperParams.OpParams->SessionName = TEXT("GetSessionByNameValidName");
+	LeaveSessionHelperParams.OpParams->bDestroySession = true;
+
+	FTestPipeline& LoginPipeline = GetLoginPipeline({ AccountId });
 
 	CreateSessionHelperParams.OpParams->LocalAccountId = AccountId;
+	LeaveSessionHelperParams.OpParams->LocalAccountId = AccountId;
 	
 	LoginPipeline
 		.EmplaceStep<FCreateSessionHelper>(MoveTemp(CreateSessionHelperParams))
@@ -150,7 +173,8 @@ SESSIONS_TEST_CASE("If I call GetSessionByName with valid information, it return
 				ISessionsPtr SessionsInterface = OnlineSubsystem->GetSessionsInterface();
 				TOnlineResult<FGetSessionByName> Result = SessionsInterface->GetSessionByName(MoveTemp(OpGetByNameParams));
 				REQUIRE_OP(Result);
-			});
+			})
+		.EmplaceStep<FLeaveSessionHelper>(MoveTemp(LeaveSessionHelperParams));
 
 	RunToCompletion();
 }
@@ -161,7 +185,7 @@ SESSIONS_TEST_CASE("If I call GetSessionById with an invalid session id, I get a
 
 	FAccountId AccountId;
 
-	GetLoginPipeline(AccountId)
+	GetLoginPipeline({ AccountId })
 		.EmplaceLambda([](SubsystemType OnlineSubsystem)
 			{
 				FGetSessionById::Params OpGetByIdParams;
@@ -178,7 +202,9 @@ SESSIONS_TEST_CASE("If I call GetSessionById with an invalid session id, I get a
 
 SESSIONS_TEST_CASE("If I call GetSessionById with a valid but unregistered session id, I get an error", EG_SESSIONS_GETSESSIONS_TAG)
 {
-	ResetAccountStatus();
+	LOG_SCOPE_VERBOSITY_OVERRIDE(LogOnlineServices, ELogVerbosity::NoLogging);
+
+	DestroyCurrentServiceModule();
 
 	FAccountId AccountId;
 
@@ -198,7 +224,7 @@ SESSIONS_TEST_CASE("If I call GetSessionById with a valid but unregistered sessi
 
 	FGetSessionById::Params OpGetByIdParams;
 
-	FTestPipeline& LoginPipeline = GetLoginPipeline(AccountId);
+	FTestPipeline& LoginPipeline = GetLoginPipeline({ AccountId });
 
 	CreateSessionHelperParams.OpParams->LocalAccountId = AccountId;
 	LeaveSessionHelperParams.OpParams->LocalAccountId = AccountId;
@@ -227,6 +253,8 @@ SESSIONS_TEST_CASE("If I call GetSessionById with a valid but unregistered sessi
 
 SESSIONS_TEST_CASE("If I call GetSessionById with a valid id for a valid session, I get a valid session reference in return", EG_SESSIONS_GETSESSIONS_TAG)
 {
+	LOG_SCOPE_VERBOSITY_OVERRIDE(LogEOSSDK, ELogVerbosity::NoLogging);
+
 	FAccountId AccountId;
 
 	FCreateSession::Params OpCreateParams;
@@ -236,11 +264,18 @@ SESSIONS_TEST_CASE("If I call GetSessionById with a valid id for a valid session
 	CreateSessionHelperParams.OpParams->SessionSettings.SchemaName = TEXT("SchemaName");
 	CreateSessionHelperParams.OpParams->SessionSettings.NumMaxConnections = 2;
 
+	FLeaveSession::Params OpLeaveParams;
+	FLeaveSessionHelper::FHelperParams LeaveSessionHelperParams;
+	LeaveSessionHelperParams.OpParams = &OpLeaveParams;
+	LeaveSessionHelperParams.OpParams->SessionName = TEXT("GetSessionByValidIdName");
+	LeaveSessionHelperParams.OpParams->bDestroySession = true;
+
 	FGetSessionById::Params OpGetByIdParams;
 
-	FTestPipeline& LoginPipeline = GetLoginPipeline(AccountId);
+	FTestPipeline& LoginPipeline = GetLoginPipeline({ AccountId });
 
 	CreateSessionHelperParams.OpParams->LocalAccountId = AccountId;
+	LeaveSessionHelperParams.OpParams->LocalAccountId = AccountId;
 
 	LoginPipeline
 		.EmplaceStep<FCreateSessionHelper>(MoveTemp(CreateSessionHelperParams))
@@ -258,17 +293,16 @@ SESSIONS_TEST_CASE("If I call GetSessionById with a valid id for a valid session
 				TOnlineResult<FGetSessionById> Result = SessionsInterface->GetSessionById(MoveTemp(OpGetByIdParams));
 				REQUIRE_OP(Result);
 				CHECK(Result.GetOkValue().Session->GetSessionId() == OpGetByIdParams.SessionId);
-			});
+			})
+		.EmplaceStep<FLeaveSessionHelper>(MoveTemp(LeaveSessionHelperParams));
+
 
 	RunToCompletion();
 }
 
 SESSIONS_TEST_CASE("If I call GetPresenceSession with an invalid id, I get an error", EG_SESSIONS_GETSESSIONS_TAG)
 {
-
-	const int32 NumUsersToLogin = 0;
-
-	GetLoginPipeline(NumUsersToLogin)
+	GetPipeline()
 		.EmplaceLambda([](SubsystemType OnlineSubsystem)
 			{
 				FGetPresenceSession::Params OpGetPresenceParams;
@@ -287,7 +321,7 @@ SESSIONS_TEST_CASE("If I call GetPresenceSession with an unregistered id, I get
 {
 	FAccountId AccountId;
 
-	GetLoginPipeline(AccountId)
+	GetLoginPipeline({ AccountId })
 		.EmplaceLambda([&AccountId](SubsystemType OnlineSubsystem)
 			{
 				FGetPresenceSession::Params OpGetPresenceParams;
@@ -305,8 +339,6 @@ SESSIONS_TEST_CASE("If I call GetPresenceSession with an unregistered id, I get
 SESSIONS_TEST_CASE("If I call GetPresenceSession with a valid id, I get a valid reference to the session", EG_SESSIONS_GETSESSIONS_TAG)
 {
 	DestroyCurrentServiceModule();
-	ResetAccountStatus();
-
 
 	FAccountId AccountId;
 
@@ -324,7 +356,7 @@ SESSIONS_TEST_CASE("If I call GetPresenceSession with a valid id, I get a valid
 	LeaveSessionHelperParams.OpParams->SessionName = TEXT("GetPresenceSessionWithValidIdName");
 	LeaveSessionHelperParams.OpParams->bDestroySession = true;
 
-	FTestPipeline& LoginPipeline = GetLoginPipeline(AccountId);
+	FTestPipeline& LoginPipeline = GetLoginPipeline({ AccountId });
 
 	CreateSessionHelperParams.OpParams->LocalAccountId = AccountId;
 	LeaveSessionHelperParams.OpParams->LocalAccountId = AccountId;
@@ -350,7 +382,7 @@ SESSIONS_TEST_CASE("If I call GetSessionInviteById with an invalid account id, I
 	FGetSessionInviteById::Params OpGetInviteByIdParams;
 	OpGetInviteByIdParams.LocalAccountId = FAccountId();
 
-	GetLoginPipeline()
+	GetPipeline()
 		.EmplaceLambda([&OpGetInviteByIdParams](SubsystemType OnlineSubsystem)
 			{
 				ISessionsPtr SessionsInterface = OnlineSubsystem->GetSessionsInterface();
@@ -371,7 +403,7 @@ SESSIONS_TEST_CASE("If I call GetSessionInviteById with an invalid session invit
 	FGetSessionInviteById::Params OpGetInviteByIdParams;
 	OpGetInviteByIdParams.SessionInviteId = FSessionInviteId();
 	
-	FTestPipeline& LoginPipeline = GetLoginPipeline(AccountId);
+	FTestPipeline& LoginPipeline = GetLoginPipeline({ AccountId });
 
 	OpGetInviteByIdParams.LocalAccountId = AccountId;
 
@@ -393,6 +425,7 @@ SESSIONS_TEST_CASE("If I call GetSessionInviteById with a valid account id, but
 {
 	DestroyCurrentServiceModule();
 
+	int32 UserNumToLogin = 7;
 	FAccountId FirstAccountId, SecondAccountId;
 
 	FCreateSession::Params OpCreateParams;
@@ -410,12 +443,19 @@ SESSIONS_TEST_CASE("If I call GetSessionInviteById with a valid account id, but
 
 	FGetSessionInviteById::Params OpGetInviteByIdParams;
 
-	FTestPipeline& LoginPipeline = GetLoginPipeline(FirstAccountId, SecondAccountId);
+	FLeaveSession::Params OpLeaveParams;
+	FLeaveSessionHelper::FHelperParams LeaveSessionHelperParams;
+	LeaveSessionHelperParams.OpParams = &OpLeaveParams;
+	LeaveSessionHelperParams.OpParams->SessionName = TEXT("GetInviteByIdWithValidAccountIdName");
+	LeaveSessionHelperParams.OpParams->bDestroySession = true;
+
+	FTestPipeline& LoginPipeline = GetLoginPipeline(UserNumToLogin, { FirstAccountId, SecondAccountId });
 
 	OpGetInviteByIdParams.LocalAccountId = FirstAccountId;
 	CreateSessionHelperParams.OpParams->LocalAccountId = FirstAccountId;
 	SendSessionInviteHelperParams.OpParams->LocalAccountId = FirstAccountId;
 	SendSessionInviteHelperParams.OpParams->TargetUsers.Add(SecondAccountId);
+	LeaveSessionHelperParams.OpParams->LocalAccountId = FirstAccountId;
 
 	LoginPipeline
 		.EmplaceStep<FCreateSessionHelper>(MoveTemp(CreateSessionHelperParams))
@@ -431,16 +471,20 @@ SESSIONS_TEST_CASE("If I call GetSessionInviteById with a valid account id, but
 				TOnlineResult<FGetSessionInviteById> Result = SessionsInterface->GetSessionInviteById(MoveTemp(OpGetInviteByIdParams));
 				REQUIRE(Result.IsError());
 				CHECK(Result.GetErrorValue() == Errors::InvalidState());
-			});
+			})
+		.EmplaceStep<FLeaveSessionHelper>(MoveTemp(LeaveSessionHelperParams));
+
 
 	RunToCompletion();
 }
 
 SESSIONS_TEST_CASE("If I call GetSessionInviteById with valid invite id, but without invite, I get an error", EG_SESSIONS_GETSESSIONSEOS_TAG)
 {
+	LOG_SCOPE_VERBOSITY_OVERRIDE(LogEOSSDK, ELogVerbosity::NoLogging);
+
 	DestroyCurrentServiceModule();
-	ResetAccountStatus();
 
+	int32 UserNumToLogin = 7;
 	FAccountId FirstAccountId, SecondAccountId;
 
 	FCreateSession::Params OpCreateParams;
@@ -462,13 +506,20 @@ SESSIONS_TEST_CASE("If I call GetSessionInviteById with valid invite id, but wit
 
 	FGetSessionInviteById::Params OpGetInviteByIdParams;
 
-	FTestPipeline& LoginPipeline = GetLoginPipeline(FirstAccountId, SecondAccountId);
+	FLeaveSession::Params OpLeaveParams;
+	FLeaveSessionHelper::FHelperParams LeaveSessionHelperParams;
+	LeaveSessionHelperParams.OpParams = &OpLeaveParams;
+	LeaveSessionHelperParams.OpParams->SessionName = TEXT("GetInviteByIdWithValidInviteIdName");
+	LeaveSessionHelperParams.OpParams->bDestroySession = true;
+
+	FTestPipeline& LoginPipeline = GetLoginPipeline(UserNumToLogin, { FirstAccountId, SecondAccountId });
 
 	OpGetInviteByIdParams.LocalAccountId = SecondAccountId;
 	CreateSessionHelperParams.OpParams->LocalAccountId = FirstAccountId;
 	SendSessionInviteHelperParams.OpParams->LocalAccountId = FirstAccountId;
 	SendSessionInviteHelperParams.OpParams->TargetUsers.Add(SecondAccountId);
 	RejectSessionInviteHelperParams.OpParams->LocalAccountId = SecondAccountId;
+	LeaveSessionHelperParams.OpParams->LocalAccountId = FirstAccountId;
 
 	LoginPipeline
 		.EmplaceStep<FCreateSessionHelper>(MoveTemp(CreateSessionHelperParams))
@@ -486,7 +537,9 @@ SESSIONS_TEST_CASE("If I call GetSessionInviteById with valid invite id, but wit
 				TOnlineResult<FGetSessionInviteById> Result = SessionsInterface->GetSessionInviteById(MoveTemp(OpGetInviteByIdParams));
 				REQUIRE(Result.IsError());
 				CHECK(Result.GetErrorValue() == Errors::NotFound());
-			});
+			})
+		.EmplaceStep<FLeaveSessionHelper>(MoveTemp(LeaveSessionHelperParams));
+
 
 	RunToCompletion();
 }
@@ -494,8 +547,8 @@ SESSIONS_TEST_CASE("If I call GetSessionInviteById with valid invite id, but wit
 SESSIONS_TEST_CASE("If I call GetSessionInviteById with a valid data, I get a valid reference to the session invite", EG_SESSIONS_GETSESSIONSEOS_TAG)
 {
 	DestroyCurrentServiceModule();
-	ResetAccountStatus();
 
+	int32 UserNumToLogin = 7;
 	FAccountId FirstAccountId, SecondAccountId;
 
 	FCreateSession::Params OpCreateParams;
@@ -513,12 +566,19 @@ SESSIONS_TEST_CASE("If I call GetSessionInviteById with a valid data, I get a va
 
 	FGetSessionInviteById::Params OpGetInviteByIdParams;
 
-	FTestPipeline& LoginPipeline = GetLoginPipeline(FirstAccountId, SecondAccountId);
+	FLeaveSession::Params OpLeaveParams;
+	FLeaveSessionHelper::FHelperParams LeaveSessionHelperParams;
+	LeaveSessionHelperParams.OpParams = &OpLeaveParams;
+	LeaveSessionHelperParams.OpParams->SessionName = TEXT("GetInviteByIdValidName");
+	LeaveSessionHelperParams.OpParams->bDestroySession = true;
+
+	FTestPipeline& LoginPipeline = GetLoginPipeline(UserNumToLogin, { FirstAccountId, SecondAccountId });
 
 	OpGetInviteByIdParams.LocalAccountId = SecondAccountId;
 	CreateSessionHelperParams.OpParams->LocalAccountId = FirstAccountId;
 	SendSessionInviteHelperParams.OpParams->LocalAccountId = FirstAccountId;
 	SendSessionInviteHelperParams.OpParams->TargetUsers.Add(SecondAccountId);
+	LeaveSessionHelperParams.OpParams->LocalAccountId = FirstAccountId;
 
 	LoginPipeline
 		.EmplaceStep<FCreateSessionHelper>(MoveTemp(CreateSessionHelperParams))
@@ -534,7 +594,8 @@ SESSIONS_TEST_CASE("If I call GetSessionInviteById with a valid data, I get a va
 				TOnlineResult<FGetSessionInviteById> Result = SessionsInterface->GetSessionInviteById(MoveTemp(OpGetInviteByIdParams));
 				REQUIRE(Result.GetOkValue().SessionInvite->GetInviteId().IsValid());
 				CHECK(Result.GetOkValue().SessionInvite->GetSenderId() == FirstAccountId);
-			});
+			})
+		.EmplaceStep<FLeaveSessionHelper>(MoveTemp(LeaveSessionHelperParams));
 
 	RunToCompletion();
 }
\ No newline at end of file
diff --git a/Engine/Source/Programs/Online/OnlineTestsCore/Source/Private/Tests/Sessions/JoinSessionTests.cpp b/Engine/Source/Programs/Online/OnlineTestsCore/Source/Private/Tests/Sessions/JoinSessionTests.cpp
--- a/Engine/Source/Programs/Online/OnlineTestsCore/Source/Private/Tests/Sessions/JoinSessionTests.cpp
+++ b/Engine/Source/Programs/Online/OnlineTestsCore/Source/Private/Tests/Sessions/JoinSessionTests.cpp
@@ -6,9 +6,11 @@
 #include "Helpers/Sessions/AddRemoveSessionMemberHelper.h"
 #include "Helpers/Sessions/JoinSessionHelper.h"
 #include "Helpers/Sessions/UpdateSessionSettingsHelper.h"
+#include "Logging/LogScopedVerbosityOverride.h"
+#include "Online/OnlineServicesLog.h"
+#include "EOSShared.h"
 #include "Helpers/LambdaStep.h"
 #include "Helpers/TickForTime.h"
-#include "OnlineCatchHelper.h"
 
 #define SESSIONS_TAG "[suite_sessions]"
 #define EG_SESSIONS_JOINSESSIONS_TAG SESSIONS_TAG "[joinsession]"
@@ -23,7 +25,7 @@ SESSIONS_TEST_CASE("If I call JoinSession with an invalid account id, I get an e
 	JoinSessionHelperParams.OpParams->LocalAccountId = FAccountId();
 	JoinSessionHelperParams.ExpectedError = TOnlineResult<FJoinSession>(Errors::InvalidParams());
 
-	GetLoginPipeline()
+	GetPipeline()
 		.EmplaceStep<FJoinSessionHelper>(MoveTemp(JoinSessionHelperParams));
 
 	RunToCompletion();
@@ -39,7 +41,7 @@ SESSIONS_TEST_CASE("If I call JoinSession with an invalid session id, I get an e
 	JoinSessionHelperParams.OpParams->SessionId = FOnlineSessionId();
 	JoinSessionHelperParams.ExpectedError = TOnlineResult<FJoinSession>(Errors::InvalidParams());
 
-	FTestPipeline& LoginPipeline = GetLoginPipeline(AccountId);
+	FTestPipeline& LoginPipeline = GetLoginPipeline({ AccountId });
 
 	JoinSessionHelperParams.OpParams->LocalAccountId = AccountId;
 
@@ -51,9 +53,11 @@ SESSIONS_TEST_CASE("If I call JoinSession with an invalid session id, I get an e
 
 SESSIONS_TEST_CASE("If I call JoinSession with an empty session name, I get an error", EG_SESSIONS_JOINSESSIONS_TAG)
 {
+	LOG_SCOPE_VERBOSITY_OVERRIDE(LogEOSSDK, ELogVerbosity::NoLogging);
+
 	DestroyCurrentServiceModule();
-	ResetAccountStatus();
 
+	int32 UserNumToLogin = 7;
 	FAccountId AccountId;
 
 	FCreateSession::Params OpCreateParams;
@@ -76,7 +80,7 @@ SESSIONS_TEST_CASE("If I call JoinSession with an empty session name, I get an e
 	JoinSessionHelperParams.OpParams = &OpJoinParams;
 	JoinSessionHelperParams.ExpectedError = TOnlineResult<FJoinSession>(Errors::InvalidParams());
 
-	FTestPipeline& LoginPipeline = GetLoginPipeline(AccountId);
+	FTestPipeline& LoginPipeline = GetLoginPipeline(UserNumToLogin, { AccountId });
 
 	CreateSessionHelperParams.OpParams->LocalAccountId = AccountId;
 	JoinSessionHelperParams.OpParams->LocalAccountId = AccountId;
@@ -100,9 +104,11 @@ SESSIONS_TEST_CASE("If I call JoinSession with an empty session name, I get an e
 
 SESSIONS_TEST_CASE("If I call JoinSession with an name already in use, I get an error", EG_SESSIONS_JOINSESSIONSEOS_TAG)
 {
+	LOG_SCOPE_VERBOSITY_OVERRIDE(LogEOSSDK, ELogVerbosity::NoLogging);
+
 	DestroyCurrentServiceModule();
-	ResetAccountStatus();
 
+	int32 UserNumToLogin = 7;
 	FAccountId FirstAccountId, SecondAccountId;
 
 	FCreateSession::Params OpFirstCreateParams;
@@ -159,7 +165,7 @@ SESSIONS_TEST_CASE("If I call JoinSession with an name already in use, I get an
 	SecondLeaveSessionHelperParams.OpParams->SessionName = TEXT("SessionNameInUseJoin1");
 	SecondLeaveSessionHelperParams.OpParams->bDestroySession = true;
 
-	FTestPipeline& LoginPipeline = GetLoginPipeline(FirstAccountId, SecondAccountId);
+	FTestPipeline& LoginPipeline = GetLoginPipeline(UserNumToLogin, { FirstAccountId, SecondAccountId });
 
 	FirstCreateSessionHelperParams.OpParams->LocalAccountId = FirstAccountId;
 	SecondCreateSessionHelperParams.OpParams->LocalAccountId = FirstAccountId;
@@ -215,8 +221,12 @@ SESSIONS_TEST_CASE("If I call JoinSession with an name already in use, I get an
 
 SESSIONS_TEST_CASE("If I call JoinSession with a valid but unregistered session id, I get an error", EG_SESSIONS_JOINSESSIONSEOS_TAG)
 {
+	LOG_SCOPE_VERBOSITY_OVERRIDE(LogEOSSDK, ELogVerbosity::NoLogging);
+	LOG_SCOPE_VERBOSITY_OVERRIDE(LogOnlineServices, ELogVerbosity::NoLogging);
+
 	DestroyCurrentServiceModule();
 
+	int32 UserNumToLogin = 7;
 	FAccountId FirstAccountId, SecondAccountId;
 
 	FCreateSession::Params OpCreateParams;
@@ -262,7 +272,7 @@ SESSIONS_TEST_CASE("If I call JoinSession with a valid but unregistered session
 	SecondLeaveSessionHelperParams.OpParams->SessionName = TEXT("SessionIdNotFoundJoinName");
 	SecondLeaveSessionHelperParams.ExpectedError = TOnlineResult<FLeaveSession>(Errors::InvalidState());
 
-	FTestPipeline& LoginPipeline = GetLoginPipeline(FirstAccountId, SecondAccountId);
+	FTestPipeline& LoginPipeline = GetLoginPipeline(UserNumToLogin, { FirstAccountId, SecondAccountId });
 
 	CreateSessionHelperParams.OpParams->LocalAccountId = FirstAccountId;
 
@@ -338,13 +348,20 @@ SESSIONS_TEST_CASE("If I call JoinSession for a session I'm already registered i
 	SecondJoinSessionHelperParams.OpParams->SessionName = TEXT("AlreadyInSessionJoinName3");
 	SecondJoinSessionHelperParams.ExpectedError = TOnlineResult<FJoinSession>(Errors::AccessDenied());
 
-	FTestPipeline& LoginPipeline = GetLoginPipeline(FirstAccountId, SecondAccountId);
+	FLeaveSession::Params OpLeaveParams;
+	FLeaveSessionHelper::FHelperParams LeaveSessionHelperParams;
+	LeaveSessionHelperParams.OpParams = &OpLeaveParams;
+	LeaveSessionHelperParams.OpParams->SessionName = TEXT("AlreadyInSessionJoinName");
+	LeaveSessionHelperParams.OpParams->bDestroySession = true;
+
+	FTestPipeline& LoginPipeline = GetLoginPipeline({ FirstAccountId, SecondAccountId });
 
 	CreateSessionHelperParams.OpParams->LocalAccountId = FirstAccountId;
 	AddSessionMemberHelperParams.OpParams->LocalAccountId = SecondAccountId;
 
 	FirstJoinSessionHelperParams.OpParams->LocalAccountId = SecondAccountId;
 	SecondJoinSessionHelperParams.OpParams->LocalAccountId = SecondAccountId;
+	LeaveSessionHelperParams.OpParams->LocalAccountId = FirstAccountId;
 
 	FindSessionsHelperParams.OpParams->LocalAccountId = SecondAccountId;
 
@@ -367,7 +384,8 @@ SESSIONS_TEST_CASE("If I call JoinSession for a session I'm already registered i
 		.EmplaceStep<FFindSessionsHelper>(MoveTemp(FindSessionsHelperParams), ExpectedSessionsFound)
 		.EmplaceStep<FJoinSessionHelper>(MoveTemp(FirstJoinSessionHelperParams))
 		.EmplaceStep<FAddSessionMemberHelper>(MoveTemp(AddSessionMemberHelperParams))
-		.EmplaceStep<FJoinSessionHelper>(MoveTemp(SecondJoinSessionHelperParams));
+		.EmplaceStep<FJoinSessionHelper>(MoveTemp(SecondJoinSessionHelperParams))
+		.EmplaceStep<FLeaveSessionHelper>(MoveTemp(LeaveSessionHelperParams));
 
 	RunToCompletion();
 }
@@ -375,8 +393,8 @@ SESSIONS_TEST_CASE("If I call JoinSession for a session I'm already registered i
 SESSIONS_TEST_CASE("If I call JoinSession for a session that is not joinable, I get an error", EG_SESSIONS_JOINSESSIONSEOS_TAG)
 {
 	DestroyCurrentServiceModule();
-	ResetAccountStatus();
 
+	int32 UserNumToLogin = 7;
 	FAccountId FirstAccountId, SecondAccountId;
 
 	FCreateSession::Params OpCreateParams;
@@ -412,7 +430,7 @@ SESSIONS_TEST_CASE("If I call JoinSession for a session that is not joinable, I
 	JoinSessionHelperParams.OpParams->SessionName = TEXT("SessionNotJoinableName1");
 	JoinSessionHelperParams.ExpectedError = TOnlineResult<FJoinSession>(Errors::AccessDenied());
 
-	FTestPipeline& LoginPipeline = GetLoginPipeline(FirstAccountId, SecondAccountId);
+	FTestPipeline& LoginPipeline = GetLoginPipeline(UserNumToLogin, { FirstAccountId, SecondAccountId });
 
 	CreateSessionHelperParams.OpParams->LocalAccountId = FirstAccountId;
 	UpdateSessionSettingsHelperParams.OpParams->LocalAccountId = FirstAccountId;
@@ -471,7 +489,7 @@ SESSIONS_TEST_CASE("If I call JoinSession with presence enabled when there is al
 	JoinSessionHelperParams.OpParams->bPresenceEnabled = true;
 	JoinSessionHelperParams.ExpectedError = TOnlineResult<FJoinSession>(Errors::InvalidState());
 
-	FTestPipeline& LoginPipeline = GetLoginPipeline(AccountId);
+	FTestPipeline& LoginPipeline = GetLoginPipeline({ AccountId });
 
 	CreateSessionHelperParams.OpParams->LocalAccountId = AccountId;
 	LeaveSessionHelperParams.OpParams->LocalAccountId = AccountId;
@@ -523,7 +541,7 @@ SESSIONS_TEST_CASE("If I call JoinSession with valid information, the operation
 	LeaveSessionHelperParams.OpParams->SessionName = TEXT("SessionValidJoinName");
 	LeaveSessionHelperParams.OpParams->bDestroySession = true;
 
-	FTestPipeline& LoginPipeline = GetLoginPipeline(FirstAccountId, SecondAccountId);
+	FTestPipeline& LoginPipeline = GetLoginPipeline({ FirstAccountId, SecondAccountId });
 
 	CreateSessionHelperParams.OpParams->LocalAccountId = FirstAccountId;
 	FindSessionsHelperParams.OpParams->LocalAccountId = SecondAccountId;
diff --git a/Engine/Source/Programs/Online/OnlineTestsCore/Source/Private/Tests/Sessions/LeaveSessionTests.cpp b/Engine/Source/Programs/Online/OnlineTestsCore/Source/Private/Tests/Sessions/LeaveSessionTests.cpp
--- a/Engine/Source/Programs/Online/OnlineTestsCore/Source/Private/Tests/Sessions/LeaveSessionTests.cpp
+++ b/Engine/Source/Programs/Online/OnlineTestsCore/Source/Private/Tests/Sessions/LeaveSessionTests.cpp
@@ -2,7 +2,6 @@
 
 #include "Helpers/Sessions/CreateSessionHelper.h"
 #include "Helpers/Sessions/LeaveSessionHelper.h"
-#include "OnlineCatchHelper.h"
 
 #define SESSIONS_TAG "[suite_sessions]"
 #define SESSIONS_LEAVESESSION_TAG SESSIONS_TAG "[leavesesssion]"
@@ -16,8 +15,8 @@ SESSIONS_TEST_CASE("If I call LeaveSession with an invalid account id, I get an
 	LeaveSessionHelperParams.OpParams->LocalAccountId = FAccountId();
 	LeaveSessionHelperParams.ExpectedError = TOnlineResult<FLeaveSession>(Errors::InvalidParams());
 
-	GetLoginPipeline()
-		.EmplaceStep<FLeaveSessionHelper>(MoveTemp(LeaveSessionHelperParams	));
+	GetPipeline()
+		.EmplaceStep<FLeaveSessionHelper>(MoveTemp(LeaveSessionHelperParams));
 
 	RunToCompletion();
 }
@@ -32,7 +31,7 @@ SESSIONS_TEST_CASE("If I call LeaveSession with an empty session name, I get an
 	LeaveSessionHelperParams.OpParams->SessionName = TEXT("");
 	LeaveSessionHelperParams.ExpectedError = TOnlineResult<FLeaveSession>(Errors::InvalidParams());
 
-	FTestPipeline& LoginPipeline = GetLoginPipeline(AccountId);
+	FTestPipeline& LoginPipeline = GetLoginPipeline({ AccountId });
 
 	LeaveSessionHelperParams.OpParams->LocalAccountId = AccountId;
 
@@ -52,7 +51,7 @@ SESSIONS_TEST_CASE("If I call LeaveSession with an unregistered session name, I
 	LeaveSessionHelperParams.OpParams->SessionName = TEXT("LeaveUnregisteredSessionName");
 	LeaveSessionHelperParams.ExpectedError = TOnlineResult<FLeaveSession>(Errors::InvalidState());
 
-	FTestPipeline& LoginPipeline = GetLoginPipeline(AccountId);
+	FTestPipeline& LoginPipeline = GetLoginPipeline({ AccountId });
 
 	LeaveSessionHelperParams.OpParams->LocalAccountId = AccountId;
 
@@ -65,8 +64,8 @@ SESSIONS_TEST_CASE("If I call LeaveSession with an unregistered session name, I
 SESSIONS_TEST_CASE("If I call LeaveSession with valid data, the operation completes successfully", SESSIONS_LEAVESESSION_TAG)
 {
 	DestroyCurrentServiceModule();
-	ResetAccountStatus();
 
+	int32 UserNumToLogin = 7;
 	FAccountId AccountId;
 
 	FCreateSession::Params OpCreateParams;
@@ -83,7 +82,7 @@ SESSIONS_TEST_CASE("If I call LeaveSession with valid data, the operation comple
 	LeaveSessionHelperParams.OpParams->SessionName = TEXT("LeaveSessionNameValidLeave");
 	LeaveSessionHelperParams.OpParams->bDestroySession = true;
 
-	FTestPipeline& LoginPipeline = GetLoginPipeline(AccountId);
+	FTestPipeline& LoginPipeline = GetLoginPipeline(UserNumToLogin, { AccountId });
 
 	CreateSessionHelperParams.OpParams->LocalAccountId = AccountId;
 	LeaveSessionHelperParams.OpParams->LocalAccountId = AccountId;
diff --git a/Engine/Source/Programs/Online/OnlineTestsCore/Source/Private/Tests/Sessions/PresenceSessionTests.cpp b/Engine/Source/Programs/Online/OnlineTestsCore/Source/Private/Tests/Sessions/PresenceSessionTests.cpp
--- a/Engine/Source/Programs/Online/OnlineTestsCore/Source/Private/Tests/Sessions/PresenceSessionTests.cpp
+++ b/Engine/Source/Programs/Online/OnlineTestsCore/Source/Private/Tests/Sessions/PresenceSessionTests.cpp
@@ -2,8 +2,10 @@
 
 #include "Helpers/Sessions/CreateSessionHelper.h"
 #include "Helpers/Sessions/LeaveSessionHelper.h"
+#include "Logging/LogScopedVerbosityOverride.h"
+#include "Online/OnlineServicesLog.h"
+#include "EOSShared.h"
 #include "Helpers/TickForTime.h"
-#include "OnlineCatchHelper.h"
 
 #define SESSIONS_TAG "[suite_sessions]"
 #define EG_SESSIONS_PRESENCESESSION_TAG SESSIONS_TAG "[presencesession]"
@@ -14,7 +16,8 @@
 SESSIONS_TEST_CASE("If I call IsPresenceSession with an invalid account id, I get an error", EG_SESSIONS_PRESENCESESSION_TAG)
 {
 	DestroyCurrentServiceModule();
-	GetLoginPipeline()
+
+	GetPipeline()
 		.EmplaceLambda([](SubsystemType OnlineSubsystem)
 			{
 				FIsPresenceSession::Params OpIsPresenceParams;
@@ -33,7 +36,7 @@ SESSIONS_TEST_CASE("If I call IsPresenceSession with an invalid session id, I ge
 {
 	FAccountId AccountId;
 
-	GetLoginPipeline(AccountId)
+	GetLoginPipeline({ AccountId })
 		.EmplaceLambda([&AccountId](SubsystemType OnlineSubsystem)
 			{
 				FIsPresenceSession::Params OpIsPresenceParams;
@@ -51,9 +54,11 @@ SESSIONS_TEST_CASE("If I call IsPresenceSession with an invalid session id, I ge
 
 SESSIONS_TEST_CASE("If I call IsPresenceSession with an unregistered account id, I get an error", EG_SESSIONS_PRESENCESESSION_TAG)
 {
+	LOG_SCOPE_VERBOSITY_OVERRIDE(LogOnlineServices, ELogVerbosity::NoLogging);
+
 	DestroyCurrentServiceModule();
-	ResetAccountStatus();
 
+	int32 UserNumToLogin = 7;
 	FAccountId AccountId;
 
 	FCreateSession::Params OpCreateParams;
@@ -75,7 +80,7 @@ SESSIONS_TEST_CASE("If I call IsPresenceSession with an unregistered account id,
 	FGetSessionByName::Params OpGetByNameParams;
 	OpGetByNameParams.LocalName = TEXT("IsPresenceUnregisteredSessionName");
 
-	FTestPipeline& LoginPipeline = GetLoginPipeline(AccountId);
+	FTestPipeline& LoginPipeline = GetLoginPipeline(UserNumToLogin, { AccountId });
 
 	CreateSessionHelperParams.OpParams->LocalAccountId = AccountId;
 	LeaveSessionHelperParams.OpParams->LocalAccountId = AccountId;
@@ -106,6 +111,8 @@ SESSIONS_TEST_CASE("If I call IsPresenceSession with an unregistered account id,
 
 SESSIONS_TEST_CASE("If I call IsPresenceSession with valid information, it returns true if the session matches, and error if it does not", EG_SESSIONS_PRESENCESESSIONEOS_TAG)
 {
+	LOG_SCOPE_VERBOSITY_OVERRIDE(LogEOSSDK, ELogVerbosity::NoLogging);
+
 	FAccountId FirstAccountId, SecondAccountId;
 
 	FCreateSession::Params OpFirstCreateParams;
@@ -127,17 +134,26 @@ SESSIONS_TEST_CASE("If I call IsPresenceSession with valid information, it retur
 	FIsPresenceSession::Params OpFirstIsPresenceParams;
 	FIsPresenceSession::Params OpSecondIsPresenceParams;
 
-	FLeaveSession::Params OpLeaveParams;
-	FLeaveSessionHelper::FHelperParams LeaveSessionHelperParams;
-	LeaveSessionHelperParams.OpParams = &OpLeaveParams;
-	LeaveSessionHelperParams.OpParams->SessionName = TEXT("IsPresenceEnableSessionName");
-	LeaveSessionHelperParams.OpParams->bDestroySession = true;
+	FLeaveSession::Params OpFirstLeaveParams;
+	FLeaveSessionHelper::FHelperParams FirstLeaveSessionHelperParams;
+	FirstLeaveSessionHelperParams.OpParams = &OpFirstLeaveParams;
+	FirstLeaveSessionHelperParams.OpParams->SessionName = TEXT("IsPresenceEnableSessionName");
+	FirstLeaveSessionHelperParams.OpParams->bDestroySession = true;
+
+	FLeaveSession::Params OpSecondLeaveParams;
+	FLeaveSessionHelper::FHelperParams SecondLeaveSessionHelperParams;
+	SecondLeaveSessionHelperParams.OpParams = &OpSecondLeaveParams;
+	SecondLeaveSessionHelperParams.OpParams->SessionName = TEXT("IsPresenceDisableSessionName");
+	SecondLeaveSessionHelperParams.OpParams->bDestroySession = true;
 
-	FTestPipeline& LoginPipeline = GetLoginPipeline(FirstAccountId, SecondAccountId);
+	FTestPipeline& LoginPipeline = GetLoginPipeline({ FirstAccountId, SecondAccountId });
 	 
 	FirstCreateSessionHelperParams.OpParams->LocalAccountId = FirstAccountId;
 	SecondCreateSessionHelperParams.OpParams->LocalAccountId = SecondAccountId;
-	LeaveSessionHelperParams.OpParams->LocalAccountId = FirstAccountId;
+	
+	FirstLeaveSessionHelperParams.OpParams->LocalAccountId = FirstAccountId;
+	SecondLeaveSessionHelperParams.OpParams->LocalAccountId = SecondAccountId;
+
 	OpFirstIsPresenceParams.LocalAccountId = FirstAccountId;
 	OpSecondIsPresenceParams.LocalAccountId = SecondAccountId;
 
@@ -169,14 +185,15 @@ SESSIONS_TEST_CASE("If I call IsPresenceSession with valid information, it retur
 				REQUIRE(Result.IsError());
 				CHECK(Result.GetErrorValue() == Errors::InvalidState());
 			})
-		.EmplaceStep<FLeaveSessionHelper>(MoveTemp(LeaveSessionHelperParams));
+		.EmplaceStep<FLeaveSessionHelper>(MoveTemp(FirstLeaveSessionHelperParams))
+		.EmplaceStep<FLeaveSessionHelper>(MoveTemp(SecondLeaveSessionHelperParams));
 
 	RunToCompletion();
 }
 
 SESSIONS_TEST_CASE("If I call SetPresenceSession for EOS, I get an error", EG_SESSIONS_PRESENCESESSIONEOS_TAG)
 {
-	GetLoginPipeline()
+	GetPipeline()
 		.EmplaceLambda([](SubsystemType OnlineSubsystem)
 			{
 				FSetPresenceSession::Params OpSetPresenceParams;
@@ -192,7 +209,7 @@ SESSIONS_TEST_CASE("If I call SetPresenceSession for EOS, I get an error", EG_SE
 
 SESSIONS_TEST_CASE("If I call SetPresenceSession with an invalid account id, I get an error", EG_SESSIONS_PRESENCESESSIONNULL_TAG)
 {
-	GetLoginPipeline()
+	GetPipeline()
 		.EmplaceLambda([](SubsystemType OnlineSubsystem)
 			{
 				FSetPresenceSession::Params OpSetPresenceParams;
@@ -210,7 +227,7 @@ SESSIONS_TEST_CASE("If I call SetPresenceSession with an invalid account id, I g
 SESSIONS_TEST_CASE("If I call SetPresenceSession with an invalid session id, I get an error", EG_SESSIONS_PRESENCESESSIONNULL_TAG)
 {
 	FAccountId AccountId;
-	GetLoginPipeline(AccountId)
+	GetLoginPipeline({ AccountId })
 		.EmplaceLambda([&AccountId](SubsystemType OnlineSubsystem)
 			{
 				FSetPresenceSession::Params OpSetPresenceParams;
@@ -246,7 +263,7 @@ SESSIONS_TEST_CASE("If I call SetPresenceSession with valid data, the operation
 
 	FSetPresenceSession::Params OpSetPresenceParams;
 
-	FTestPipeline& LoginPipeline = GetLoginPipeline(AccountId);
+	FTestPipeline& LoginPipeline = GetLoginPipeline({ AccountId });
 
 	CreateSessionHelperParams.OpParams->LocalAccountId = AccountId;
 	LeaveSessionHelperParams.OpParams->LocalAccountId = AccountId;
@@ -276,7 +293,7 @@ SESSIONS_TEST_CASE("If I call SetPresenceSession with valid data, the operation
 
 SESSIONS_TEST_CASE("If I call ClearPresenceSession for EOS, I get an error", EG_SESSIONS_PRESENCESESSIONEOS_TAG)
 {
-	GetLoginPipeline()
+	GetPipeline()
 		.EmplaceLambda([](SubsystemType OnlineSubsystem)
 			{
 				FClearPresenceSession::Params OpClearPresenceParams;
@@ -292,7 +309,7 @@ SESSIONS_TEST_CASE("If I call ClearPresenceSession for EOS, I get an error", EG_
 
 SESSIONS_TEST_CASE("If I call ClearPresenceSession with an invalid account id, I get an error", EG_SESSIONS_PRESENCESESSIONNULL_TAG)
 {
-	GetLoginPipeline()
+	GetPipeline()
 		.EmplaceLambda([](SubsystemType OnlineSubsystem)
 			{
 				FClearPresenceSession::Params OpClearPresenceParams;
@@ -319,9 +336,16 @@ SESSIONS_TEST_CASE("If I call ClearPresenceSession with valid data, the operatio
 	CreateSessionHelperParams.OpParams->SessionSettings.SchemaName = TEXT("SchemaName");
 	CreateSessionHelperParams.OpParams->SessionSettings.NumMaxConnections = 2;
 
-	FTestPipeline& LoginPipeline = GetLoginPipeline(AccountId);
+	FLeaveSession::Params OpLeaveParams;
+	FLeaveSessionHelper::FHelperParams LeaveSessionHelperParams;
+	LeaveSessionHelperParams.OpParams = &OpLeaveParams;
+	LeaveSessionHelperParams.OpParams->SessionName = TEXT("ClearPresenceValidName");
+	LeaveSessionHelperParams.OpParams->bDestroySession = true;
+
+	FTestPipeline& LoginPipeline = GetLoginPipeline({ AccountId });
 	
 	CreateSessionHelperParams.OpParams->LocalAccountId = AccountId;
+	LeaveSessionHelperParams.OpParams->LocalAccountId = AccountId;
 
 	LoginPipeline
 		.EmplaceStep<FCreateSessionHelper>(MoveTemp(CreateSessionHelperParams))
@@ -333,7 +357,8 @@ SESSIONS_TEST_CASE("If I call ClearPresenceSession with valid data, the operatio
 				ISessionsPtr SessionsInterface = OnlineSubsystem->GetSessionsInterface();
 				TOnlineResult<FClearPresenceSession> Result = SessionsInterface->ClearPresenceSession(MoveTemp(OpClearPresenceParams));
 				REQUIRE_OP(Result);
-			});
+			})
+		.EmplaceStep<FLeaveSessionHelper>(MoveTemp(LeaveSessionHelperParams));
 
 	RunToCompletion();
 }
\ No newline at end of file
diff --git a/Engine/Source/Programs/Online/OnlineTestsCore/Source/Private/Tests/Sessions/SendRejectSessionInvitesTests.cpp b/Engine/Source/Programs/Online/OnlineTestsCore/Source/Private/Tests/Sessions/SendRejectSessionInvitesTests.cpp
--- a/Engine/Source/Programs/Online/OnlineTestsCore/Source/Private/Tests/Sessions/SendRejectSessionInvitesTests.cpp
+++ b/Engine/Source/Programs/Online/OnlineTestsCore/Source/Private/Tests/Sessions/SendRejectSessionInvitesTests.cpp
@@ -2,8 +2,10 @@
 
 #include "Helpers/Sessions/CreateSessionHelper.h"
 #include "Helpers/Sessions/SendRejectSessionInviteHelper.h"
+#include "Helpers/Sessions/LeaveSessionHelper.h"
+#include "Logging/LogScopedVerbosityOverride.h"
+#include "EOSShared.h"
 #include "Helpers/TickForTime.h"
-#include "OnlineCatchHelper.h"
 
 #define SESSIONS_TAG "[suite_sessions]"
 #define EG_SESSIONS_SENDREJECTSESSIONINVITE_TAG SESSIONS_TAG "[sendrejectsesssioninvite]"
@@ -21,7 +23,7 @@ SESSIONS_TEST_CASE("If I call SendSessionInvite with an invalid account id, I ge
 	SendSessionInviteHelperParams.OpParams->LocalAccountId = FAccountId();
 	SendSessionInviteHelperParams.ExpectedError = TOnlineResult<FSendSessionInvite>(Errors::InvalidParams());
 
-	GetLoginPipeline()
+	GetPipeline()
 		.EmplaceStep<FSendSessionInviteHelper>(MoveTemp(SendSessionInviteHelperParams));
 
 	RunToCompletion();
@@ -37,7 +39,7 @@ SESSIONS_TEST_CASE("If I call SendSessionInvite with an empty session name, I ge
 	SendSessionInviteHelperParams.OpParams->SessionName = TEXT("");
 	SendSessionInviteHelperParams.ExpectedError = TOnlineResult<FSendSessionInvite>(Errors::InvalidParams());
 	
-	FTestPipeline& LoginPipeline = GetLoginPipeline(AccountId);
+	FTestPipeline& LoginPipeline = GetLoginPipeline({ AccountId });
 
 	SendSessionInviteHelperParams.OpParams->LocalAccountId = AccountId;
 
@@ -57,7 +59,7 @@ SESSIONS_TEST_CASE("If I call SendSessionInvite with a valid session name but un
 	SendSessionInviteHelperParams.OpParams->SessionName = TEXT("SessionSendInviteUnregisteredName");
 	SendSessionInviteHelperParams.ExpectedError = TOnlineResult<FSendSessionInvite>(Errors::InvalidState());
 
-	FTestPipeline& LoginPipeline = GetLoginPipeline(FirstAccountId, SecondAccountId);
+	FTestPipeline& LoginPipeline = GetLoginPipeline({ FirstAccountId, SecondAccountId });
 
 	SendSessionInviteHelperParams.OpParams->LocalAccountId = FirstAccountId;
 	SendSessionInviteHelperParams.OpParams->TargetUsers.Add(SecondAccountId);
@@ -71,8 +73,8 @@ SESSIONS_TEST_CASE("If I call SendSessionInvite with a valid session name but un
 SESSIONS_TEST_CASE("If I call SendSessionInvite with an empty target users, I get an error", EG_SESSIONS_SENDREJECTSESSIONINVITE_TAG)
 {
 	DestroyCurrentServiceModule();
-	ResetAccountStatus();
 
+	int32 UserNumToLogin = 7;
 	FAccountId FirstAccountId, SecondAccountId;
 
 	FCreateSession::Params OpCreateParams;
@@ -89,14 +91,22 @@ SESSIONS_TEST_CASE("If I call SendSessionInvite with an empty target users, I ge
 	SendSessionInviteHelperParams.OpParams->SessionName = TEXT("SessionSendInviteEmptyName");
 	SendSessionInviteHelperParams.ExpectedError = TOnlineResult<FSendSessionInvite>(Errors::InvalidParams());
 	
-	FTestPipeline& LoginPipeline = GetLoginPipeline(FirstAccountId, SecondAccountId);
+	FLeaveSession::Params OpLeaveParams;
+	FLeaveSessionHelper::FHelperParams LeaveSessionHelperParams;
+	LeaveSessionHelperParams.OpParams = &OpLeaveParams;
+	LeaveSessionHelperParams.OpParams->SessionName = TEXT("SessionSendInviteEmptyName");
+	LeaveSessionHelperParams.OpParams->bDestroySession = true;
+
+	FTestPipeline& LoginPipeline = GetLoginPipeline(UserNumToLogin, { FirstAccountId, SecondAccountId });
 
 	CreateSessionHelperParams.OpParams->LocalAccountId = FirstAccountId;
 	SendSessionInviteHelperParams.OpParams->LocalAccountId = FirstAccountId;
+	LeaveSessionHelperParams.OpParams->LocalAccountId = FirstAccountId;
 
 	LoginPipeline
 		.EmplaceStep<FCreateSessionHelper>(MoveTemp(CreateSessionHelperParams))
-		.EmplaceStep<FSendSessionInviteHelper>(MoveTemp(SendSessionInviteHelperParams));
+		.EmplaceStep<FSendSessionInviteHelper>(MoveTemp(SendSessionInviteHelperParams))
+		.EmplaceStep<FLeaveSessionHelper>(MoveTemp(LeaveSessionHelperParams));
 
 	RunToCompletion();
 }
@@ -104,10 +114,10 @@ SESSIONS_TEST_CASE("If I call SendSessionInvite with an empty target users, I ge
 SESSIONS_TEST_CASE("If I call SendSessionInvite with an invalid target users, I get an error", EG_SESSIONS_SENDREJECTSESSIONINVITE_TAG)
 {
 	DestroyCurrentServiceModule();
-	ResetAccountStatus();
 
+	int32 UserNumToLogin = 7;
 	FAccountId FirstAccountId, SecondAccountId;
-
+		
 	FCreateSession::Params OpCreateParams;
 	FCreateSessionHelper::FHelperParams CreateSessionHelperParams;
 	CreateSessionHelperParams.OpParams = &OpCreateParams;
@@ -123,14 +133,22 @@ SESSIONS_TEST_CASE("If I call SendSessionInvite with an invalid target users, I
 	SendSessionInviteHelperParams.ExpectedError = TOnlineResult<FSendSessionInvite>(Errors::InvalidParams());
 	SendSessionInviteHelperParams.OpParams->TargetUsers.Add(FAccountId());
 
-	FTestPipeline& LoginPipeline = GetLoginPipeline(FirstAccountId, SecondAccountId);
+	FLeaveSession::Params OpLeaveParams;
+	FLeaveSessionHelper::FHelperParams LeaveSessionHelperParams;
+	LeaveSessionHelperParams.OpParams = &OpLeaveParams;
+	LeaveSessionHelperParams.OpParams->SessionName = TEXT("SessionSendInviteInvalidUsersName");
+	LeaveSessionHelperParams.OpParams->bDestroySession = true;
+
+	FTestPipeline& LoginPipeline = GetLoginPipeline(UserNumToLogin, { FirstAccountId, SecondAccountId });
 
 	CreateSessionHelperParams.OpParams->LocalAccountId = FirstAccountId;
 	SendSessionInviteHelperParams.OpParams->LocalAccountId = FirstAccountId;
+	LeaveSessionHelperParams.OpParams->LocalAccountId = FirstAccountId;
 
 	LoginPipeline
 		.EmplaceStep<FCreateSessionHelper>(MoveTemp(CreateSessionHelperParams))
-		.EmplaceStep<FSendSessionInviteHelper>(MoveTemp(SendSessionInviteHelperParams));
+		.EmplaceStep<FSendSessionInviteHelper>(MoveTemp(SendSessionInviteHelperParams))
+		.EmplaceStep<FLeaveSessionHelper>(MoveTemp(LeaveSessionHelperParams));
 
 	RunToCompletion();
 }
@@ -138,8 +156,8 @@ SESSIONS_TEST_CASE("If I call SendSessionInvite with an invalid target users, I
 SESSIONS_TEST_CASE("If I call SendSessionInvite with valid data, the operation completes successfully", EG_SESSIONS_SENDREJECTSESSIONINVITEEOS_TAG)
 {
 	DestroyCurrentServiceModule();
-	ResetAccountStatus();
 
+	int32 UserNumToLogin = 7;
 	FAccountId FirstAccountId, SecondAccountId;
 
 	FCreateSession::Params OpCreateParams;
@@ -157,12 +175,19 @@ SESSIONS_TEST_CASE("If I call SendSessionInvite with valid data, the operation c
 
 	FGetAllSessionInvites::Params OpGetAllSessionInvitesParams;
 
-	FTestPipeline& LoginPipeline = GetLoginPipeline(FirstAccountId, SecondAccountId);
+	FLeaveSession::Params OpLeaveParams;
+	FLeaveSessionHelper::FHelperParams LeaveSessionHelperParams;
+	LeaveSessionHelperParams.OpParams = &OpLeaveParams;
+	LeaveSessionHelperParams.OpParams->SessionName = TEXT("SessionSendInviteValidName");
+	LeaveSessionHelperParams.OpParams->bDestroySession = true;
+
+	FTestPipeline& LoginPipeline = GetLoginPipeline(UserNumToLogin, { FirstAccountId, SecondAccountId });
 
 	OpGetAllSessionInvitesParams.LocalAccountId = SecondAccountId;
 	CreateSessionHelperParams.OpParams->LocalAccountId = FirstAccountId;
 	SendSessionInviteHelperParams.OpParams->LocalAccountId = FirstAccountId;
 	SendSessionInviteHelperParams.OpParams->TargetUsers.Add(SecondAccountId);
+	LeaveSessionHelperParams.OpParams->LocalAccountId = FirstAccountId;
 
 	const uint32_t ExpectedSessionInvitesNum = 1;
 
@@ -182,7 +207,8 @@ SESSIONS_TEST_CASE("If I call SendSessionInvite with valid data, the operation c
 				ISessionsPtr SessionsInterface = OnlineSubsystem->GetSessionsInterface();
 				TOnlineResult<FGetAllSessionInvites> Result = SessionsInterface->GetAllSessionInvites(MoveTemp(OpGetAllSessionInvitesParams));
 				CHECK(Result.GetOkValue().SessionInvites.Num() == ExpectedSessionInvitesNum);
-			});
+			})
+		.EmplaceStep<FLeaveSessionHelper>(MoveTemp(LeaveSessionHelperParams));
 
 	RunToCompletion();
 }
@@ -195,7 +221,7 @@ SESSIONS_TEST_CASE("If I call RejectSessionInvite with an invalid account id, I
 	RejectSessionInviteHelperParams.OpParams->LocalAccountId = FAccountId();
 	RejectSessionInviteHelperParams.ExpectedError = TOnlineResult<FRejectSessionInvite>(Errors::InvalidParams());
 
-	GetLoginPipeline()
+	GetPipeline()
 		.EmplaceStep<FRejectSessionInviteHelper>(MoveTemp(RejectSessionInviteHelperParams));
 
 	RunToCompletion();
@@ -211,7 +237,7 @@ SESSIONS_TEST_CASE("If I call RejectSessionInvite with an invalid session invite
 	RejectSessionInviteHelperParams.OpParams->SessionInviteId = FSessionInviteId();
 	RejectSessionInviteHelperParams.ExpectedError = TOnlineResult<FRejectSessionInvite>(Errors::InvalidParams());
 	
-	FTestPipeline& LoginPipeline = GetLoginPipeline(AccountId);
+	FTestPipeline& LoginPipeline = GetLoginPipeline({ AccountId });
 	
 	RejectSessionInviteHelperParams.OpParams->LocalAccountId = AccountId;
 	
@@ -223,9 +249,11 @@ SESSIONS_TEST_CASE("If I call RejectSessionInvite with an invalid session invite
 
 SESSIONS_TEST_CASE("If I call RejectSessionInvite with valid data, the operation completes successfully", EG_SESSIONS_SENDREJECTSESSIONINVITEEOS_TAG)
 {
+	LOG_SCOPE_VERBOSITY_OVERRIDE(LogEOSSDK, ELogVerbosity::NoLogging);
+
 	DestroyCurrentServiceModule();
-	ResetAccountStatus();
 
+	int32 UserNumToLogin = 7;
 	FAccountId FirstAccountId, SecondAccountId;
 
 	FCreateSession::Params OpCreateParams;
@@ -247,7 +275,13 @@ SESSIONS_TEST_CASE("If I call RejectSessionInvite with valid data, the operation
 
 	FGetAllSessionInvites::Params OpGetAllSessionInvitesParams;
 
-	FTestPipeline& LoginPipeline = GetLoginPipeline(FirstAccountId, SecondAccountId);
+	FLeaveSession::Params OpLeaveParams;
+	FLeaveSessionHelper::FHelperParams LeaveSessionHelperParams;
+	LeaveSessionHelperParams.OpParams = &OpLeaveParams;
+	LeaveSessionHelperParams.OpParams->SessionName = TEXT("SessionRejectInviteValidName");
+	LeaveSessionHelperParams.OpParams->bDestroySession = true;
+
+	FTestPipeline& LoginPipeline = GetLoginPipeline(UserNumToLogin, { FirstAccountId, SecondAccountId });
 
 	OpGetAllSessionInvitesParams.LocalAccountId = SecondAccountId;
 	CreateSessionHelperParams.OpParams->LocalAccountId = FirstAccountId;
@@ -256,6 +290,8 @@ SESSIONS_TEST_CASE("If I call RejectSessionInvite with valid data, the operation
 	SendSessionInviteHelperParams.OpParams->LocalAccountId = FirstAccountId;
 	SendSessionInviteHelperParams.OpParams->TargetUsers.Add(SecondAccountId);
 
+	LeaveSessionHelperParams.OpParams->LocalAccountId = FirstAccountId;
+
 	const uint32_t ExpectedSessionInvitesNum = 1;
 
 	LoginPipeline
@@ -279,7 +315,8 @@ SESSIONS_TEST_CASE("If I call RejectSessionInvite with valid data, the operation
 				ISessionsPtr SessionsInterface = OnlineSubsystem->GetSessionsInterface();
 				TOnlineResult<FGetAllSessionInvites> Result = SessionsInterface->GetAllSessionInvites(MoveTemp(OpGetAllSessionInvitesParams));
 				CHECK(Result.GetOkValue().SessionInvites.IsEmpty());
-			});
+			})
+		.EmplaceStep<FLeaveSessionHelper>(MoveTemp(LeaveSessionHelperParams));
 			
 	RunToCompletion();
 }
diff --git a/Engine/Source/Programs/Online/OnlineTestsCore/Source/Private/Tests/Sessions/StartMatchmakingTests.cpp b/Engine/Source/Programs/Online/OnlineTestsCore/Source/Private/Tests/Sessions/StartMatchmakingTests.cpp
--- a/Engine/Source/Programs/Online/OnlineTestsCore/Source/Private/Tests/Sessions/StartMatchmakingTests.cpp
+++ b/Engine/Source/Programs/Online/OnlineTestsCore/Source/Private/Tests/Sessions/StartMatchmakingTests.cpp
@@ -2,7 +2,9 @@
 
 #include "Helpers/Sessions/CreateSessionHelper.h"
 #include "Helpers/Sessions/StartMatchmakingSessionHelper.h"
-#include "OnlineCatchHelper.h"
+#include "Helpers/Sessions/LeaveSessionHelper.h"
+#include "Logging/LogScopedVerbosityOverride.h"
+#include "Online/OnlineServicesLog.h"
 
 #define SESSIONS_TAG "[suite_sessions]"
 #define EG_SESSIONS_STARTMATCHMAKING_TAG SESSIONS_TAG "[startmatchmaking]"
@@ -18,7 +20,7 @@ SESSIONS_TEST_CASE("If I call StartMatchmaking with an invalid account id, I get
 	StartMatchmakingHelperParams.OpParams->SessionCreationParameters.LocalAccountId = FAccountId();
 	StartMatchmakingHelperParams.ExpectedError = TOnlineResult<FStartMatchmaking>(Errors::InvalidParams());
 
-	GetLoginPipeline()
+	GetPipeline()
 		.EmplaceStep<FStartMatchmakingHelper>(MoveTemp(StartMatchmakingHelperParams));
 		
 	RunToCompletion();
@@ -34,7 +36,7 @@ SESSIONS_TEST_CASE("If I call StartMatchmaking with an empty session name, I get
 	StartMatchmakingHelperParams.OpParams->SessionCreationParameters.SessionName = TEXT("");
 	StartMatchmakingHelperParams.ExpectedError = TOnlineResult<FStartMatchmaking>(Errors::InvalidParams());
 
-	FTestPipeline& LoginPipeline = GetLoginPipeline(AccountId);
+	FTestPipeline& LoginPipeline = GetLoginPipeline({ AccountId });
 
 	StartMatchmakingHelperParams.OpParams->SessionCreationParameters.LocalAccountId = AccountId;
 
@@ -55,7 +57,7 @@ SESSIONS_TEST_CASE("If I call StartMatchmaking with an empty schema name in sett
 	StartMatchmakingHelperParams.OpParams->SessionCreationParameters.SessionSettings.SchemaName = TEXT("");
 	StartMatchmakingHelperParams.ExpectedError = TOnlineResult<FStartMatchmaking>(Errors::InvalidParams());
 
-	FTestPipeline& LoginPipeline = GetLoginPipeline(AccountId);
+	FTestPipeline& LoginPipeline = GetLoginPipeline({ AccountId });
 
 	StartMatchmakingHelperParams.OpParams->SessionCreationParameters.LocalAccountId = AccountId;
 
@@ -77,7 +79,7 @@ SESSIONS_TEST_CASE("If I call StartMatchmaking with an invalid max connections n
 	StartMatchmakingHelperParams.OpParams->SessionCreationParameters.SessionSettings.NumMaxConnections = 0;
 	StartMatchmakingHelperParams.ExpectedError = TOnlineResult<FStartMatchmaking>(Errors::InvalidParams());
 
-	FTestPipeline& LoginPipeline = GetLoginPipeline(AccountId);
+	FTestPipeline& LoginPipeline = GetLoginPipeline({ AccountId });
 
 	StartMatchmakingHelperParams.OpParams->SessionCreationParameters.LocalAccountId = AccountId;
 
@@ -100,7 +102,7 @@ SESSIONS_TEST_CASE("If I call StartMatchmaking with an empty custom setting name
 	StartMatchmakingHelperParams.OpParams->SessionCreationParameters.SessionSettings.CustomSettings.Emplace(FName(""), FCustomSessionSetting{FSchemaVariant(false), ESchemaAttributeVisibility::Public});
 	StartMatchmakingHelperParams.ExpectedError = TOnlineResult<FStartMatchmaking>(Errors::InvalidParams());
 
-	FTestPipeline& LoginPipeline = GetLoginPipeline(AccountId);
+	FTestPipeline& LoginPipeline = GetLoginPipeline({ AccountId });
 
 	StartMatchmakingHelperParams.OpParams->SessionCreationParameters.LocalAccountId = AccountId;
 
@@ -123,7 +125,7 @@ SESSIONS_TEST_CASE("If I call StartMatchmaking with an empty search filter key,
 	StartMatchmakingHelperParams.OpParams->SessionSearchFilters.Emplace(FFindSessionsSearchFilter{ FName(""), ESchemaAttributeComparisonOp::Equals, FSchemaVariant(100.0)});
 	StartMatchmakingHelperParams.ExpectedError = TOnlineResult<FStartMatchmaking>(Errors::InvalidParams());
 
-	FTestPipeline& LoginPipeline = GetLoginPipeline(AccountId);
+	FTestPipeline& LoginPipeline = GetLoginPipeline({ AccountId });
 
 	StartMatchmakingHelperParams.OpParams->SessionCreationParameters.LocalAccountId = AccountId;
 
@@ -135,6 +137,8 @@ SESSIONS_TEST_CASE("If I call StartMatchmaking with an empty search filter key,
 
 SESSIONS_TEST_CASE("If I call StartMatchmaking with session override id less than 16 chars or more than 64, I get an error", EG_SESSIONS_STARTMATCHMAKINGEOS_TAG)
 {
+	LOG_SCOPE_VERBOSITY_OVERRIDE(LogOnlineServices, ELogVerbosity::NoLogging);
+
 	DestroyCurrentServiceModule();
 
 	FAccountId AccountId;
@@ -157,7 +161,7 @@ SESSIONS_TEST_CASE("If I call StartMatchmaking with session override id less tha
 	SecondStartMatchmakingHelperParams.OpParams->SessionCreationParameters.SessionIdOverride = TEXT("SessionIdOverrideSessionIdOverrideSessionIdOverrideSessionIdOverride");
 	SecondStartMatchmakingHelperParams.ExpectedError = TOnlineResult<FStartMatchmaking>(Errors::InvalidParams());
 
-	FTestPipeline& LoginPipeline = GetLoginPipeline(AccountId);
+	FTestPipeline& LoginPipeline = GetLoginPipeline({ AccountId });
 
 	FirstStartMatchmakingHelperParams.OpParams->SessionCreationParameters.LocalAccountId = AccountId;
 	SecondStartMatchmakingHelperParams.OpParams->SessionCreationParameters.LocalAccountId = AccountId;
@@ -172,8 +176,8 @@ SESSIONS_TEST_CASE("If I call StartMatchmaking with session override id less tha
 SESSIONS_TEST_CASE("If I call StartMatchmaking with name already in use, I get an error", EG_SESSIONS_STARTMATCHMAKING_TAG)
 {
 	DestroyCurrentServiceModule();
-	ResetAccountStatus();
 
+	int32 UserNumToLogin = 7;
 	FAccountId AccountId;
 
 	FCreateSession::Params OpCreateParams;
@@ -192,14 +196,22 @@ SESSIONS_TEST_CASE("If I call StartMatchmaking with name already in use, I get a
 	StartMatchmakingHelperParams.OpParams->SessionCreationParameters.SessionSettings.NumMaxConnections = 2;
 	StartMatchmakingHelperParams.ExpectedError = TOnlineResult<FStartMatchmaking>(Errors::InvalidState());
 
-	FTestPipeline& LoginPipeline = GetLoginPipeline(AccountId);
+	FLeaveSession::Params OpLeaveParams;
+	FLeaveSessionHelper::FHelperParams LeaveSessionHelperParams;
+	LeaveSessionHelperParams.OpParams = &OpLeaveParams;
+	LeaveSessionHelperParams.OpParams->SessionName = TEXT("StartMatchmakingAlreadyInUseSessionName");
+	LeaveSessionHelperParams.OpParams->bDestroySession = true;
+
+	FTestPipeline& LoginPipeline = GetLoginPipeline(UserNumToLogin, { AccountId });
 
 	CreateSessionHelperParams.OpParams->LocalAccountId = AccountId;
 	StartMatchmakingHelperParams.OpParams->SessionCreationParameters.LocalAccountId = AccountId;
+	LeaveSessionHelperParams.OpParams->LocalAccountId = AccountId;
 
 	LoginPipeline
 		.EmplaceStep<FCreateSessionHelper>(MoveTemp(CreateSessionHelperParams))
-		.EmplaceStep<FStartMatchmakingHelper>(MoveTemp(StartMatchmakingHelperParams));
+		.EmplaceStep<FStartMatchmakingHelper>(MoveTemp(StartMatchmakingHelperParams))
+		.EmplaceStep<FLeaveSessionHelper>(MoveTemp(LeaveSessionHelperParams));
 
 	RunToCompletion();
 }
@@ -216,7 +228,7 @@ SESSIONS_TEST_CASE("If I call StartMatchmaking with valid data, the operation wi
 	StartMatchmakingHelperParams.OpParams->SessionCreationParameters.SessionSettings.NumMaxConnections = 2;
 	StartMatchmakingHelperParams.ExpectedError = TOnlineResult<FStartMatchmaking>(Errors::NotImplemented());
 
-	FTestPipeline& LoginPipeline = GetLoginPipeline(AccountId);
+	FTestPipeline& LoginPipeline = GetLoginPipeline({ AccountId });
 
 	StartMatchmakingHelperParams.OpParams->SessionCreationParameters.LocalAccountId = AccountId;
 
diff --git a/Engine/Source/Programs/Online/OnlineTestsCore/Source/Private/Tests/Sessions/UpdateSessionSettingsTests.cpp b/Engine/Source/Programs/Online/OnlineTestsCore/Source/Private/Tests/Sessions/UpdateSessionSettingsTests.cpp
--- a/Engine/Source/Programs/Online/OnlineTestsCore/Source/Private/Tests/Sessions/UpdateSessionSettingsTests.cpp
+++ b/Engine/Source/Programs/Online/OnlineTestsCore/Source/Private/Tests/Sessions/UpdateSessionSettingsTests.cpp
@@ -2,7 +2,7 @@
 
 #include "Helpers/Sessions/CreateSessionHelper.h"
 #include "Helpers/Sessions/UpdateSessionSettingsHelper.h"
-#include "OnlineCatchHelper.h"
+#include "Helpers/Sessions/LeaveSessionHelper.h"
 #include "Algo/Count.h"
 
 #define SESSIONS_TAGS "[suite_sessions]"
@@ -17,7 +17,7 @@ UPDATESESSIONS_TEST_CASE("If I call UpdateSessionSettings with an invalid accoun
 	UpdateSessionSettingsHelperParams.OpParams->LocalAccountId = FAccountId();
 	UpdateSessionSettingsHelperParams.ExpectedError = TOnlineResult<FUpdateSessionSettings>(Errors::InvalidParams());
 
-	GetLoginPipeline()
+	GetPipeline()
 		.EmplaceStep<FUpdateSessionSettingsHelper>(MoveTemp(UpdateSessionSettingsHelperParams));
 
 	RunToCompletion();
@@ -33,7 +33,7 @@ UPDATESESSIONS_TEST_CASE("If I call UpdateSessionSettings with an empty session
 	UpdateSessionSettingsHelperParams.OpParams->SessionName = TEXT("");
 	UpdateSessionSettingsHelperParams.ExpectedError = TOnlineResult<FUpdateSessionSettings>(Errors::InvalidParams());
 
-	FTestPipeline& LoginPipeline = GetLoginPipeline(AccountId);
+	FTestPipeline& LoginPipeline = GetLoginPipeline({ AccountId });
 
 	UpdateSessionSettingsHelperParams.OpParams->LocalAccountId = AccountId;
 
@@ -54,7 +54,7 @@ UPDATESESSIONS_TEST_CASE("If I call UpdateSessionSettings with an empty schema n
 	UpdateSessionSettingsHelperParams.OpParams->Mutations.SchemaName = TEXT("");
 	UpdateSessionSettingsHelperParams.ExpectedError = TOnlineResult<FUpdateSessionSettings>(Errors::InvalidParams());
 
-	FTestPipeline& LoginPipeline = GetLoginPipeline(AccountId);
+	FTestPipeline& LoginPipeline = GetLoginPipeline({ AccountId });
 
 	UpdateSessionSettingsHelperParams.OpParams->LocalAccountId = AccountId;
 
@@ -75,7 +75,7 @@ UPDATESESSIONS_TEST_CASE("If I call UpdateSessionSettings with an invalid max co
 	UpdateSessionSettingsHelperParams.OpParams->Mutations.NumMaxConnections = 0;
 	UpdateSessionSettingsHelperParams.ExpectedError = TOnlineResult<FUpdateSessionSettings>(Errors::InvalidParams());
 
-	FTestPipeline& LoginPipeline = GetLoginPipeline(AccountId);
+	FTestPipeline& LoginPipeline = GetLoginPipeline({ AccountId });
 
 	UpdateSessionSettingsHelperParams.OpParams->LocalAccountId = AccountId;
 
@@ -96,7 +96,7 @@ UPDATESESSIONS_TEST_CASE("If I call UpdateSessionSettings with an empty custom s
 	UpdateSessionSettingsHelperParams.OpParams->Mutations.UpdatedCustomSettings.Emplace(TEXT(""), FCustomSessionSetting{FSchemaVariant(false), ESchemaAttributeVisibility::Public});
 	UpdateSessionSettingsHelperParams.ExpectedError = TOnlineResult<FUpdateSessionSettings>(Errors::InvalidParams());
 
-	FTestPipeline& LoginPipeline = GetLoginPipeline(AccountId);
+	FTestPipeline& LoginPipeline = GetLoginPipeline({ AccountId });
 
 	UpdateSessionSettingsHelperParams.OpParams->LocalAccountId = AccountId;
 
@@ -116,7 +116,7 @@ UPDATESESSIONS_TEST_CASE("If I call UpdateSessionSettings with an unregistered s
 	UpdateSessionSettingsHelperParams.OpParams->SessionName = TEXT("UnregisteredUpdateSessionName");
 	UpdateSessionSettingsHelperParams.ExpectedError = TOnlineResult<FUpdateSessionSettings>(Errors::InvalidState());
 
-	FTestPipeline& LoginPipeline = GetLoginPipeline(AccountId);
+	FTestPipeline& LoginPipeline = GetLoginPipeline({ AccountId });
 
 	UpdateSessionSettingsHelperParams.OpParams->LocalAccountId = AccountId;
 
@@ -129,8 +129,8 @@ UPDATESESSIONS_TEST_CASE("If I call UpdateSessionSettings with an unregistered s
 UPDATESESSIONS_TEST_CASE("If I call UpdateSessionSettings with valid data, the operation completes successfully", SESSIONS_UPDATESESSIONSETTINGS_TAG)
 {
 	DestroyCurrentServiceModule();
-	ResetAccountStatus();
 
+	int32 UserNumToLogin = 7;
 	FAccountId AccountId;
 
 	FCreateSession::Params OpCreateParams;
@@ -158,10 +158,17 @@ UPDATESESSIONS_TEST_CASE("If I call UpdateSessionSettings with valid data, the o
 	FGetSessionByName::Params OpGetByNameParams;
 	OpGetByNameParams.LocalName = TEXT("ValidUpdateSessionName");
 
-	FTestPipeline& LoginPipeline = GetLoginPipeline(AccountId);
+	FLeaveSession::Params OpLeaveParams;
+	FLeaveSessionHelper::FHelperParams LeaveSessionHelperParams;
+	LeaveSessionHelperParams.OpParams = &OpLeaveParams;
+	LeaveSessionHelperParams.OpParams->SessionName = TEXT("ValidUpdateSessionName");
+	LeaveSessionHelperParams.OpParams->bDestroySession = true;
+
+	FTestPipeline& LoginPipeline = GetLoginPipeline(UserNumToLogin, { AccountId });
 
 	CreateSessionHelperParams.OpParams->LocalAccountId = AccountId;
 	UpdateSessionSettingsHelperParams.OpParams->LocalAccountId = AccountId;
+	LeaveSessionHelperParams.OpParams->LocalAccountId = AccountId;
 
 	LoginPipeline
 		.EmplaceStep<FCreateSessionHelper>(MoveTemp(CreateSessionHelperParams))
@@ -185,7 +192,8 @@ UPDATESESSIONS_TEST_CASE("If I call UpdateSessionSettings with valid data, the o
 				CHECK(SessionSettings.JoinPolicy == UpdateSessionSettingsHelperParams.OpParams->Mutations.JoinPolicy);
 				CHECK(SessionSettings.bAllowNewMembers == UpdateSessionSettingsHelperParams.OpParams->Mutations.bAllowNewMembers);
 				CHECK(SessionSettings.NumMaxConnections == UpdateSessionSettingsHelperParams.OpParams->Mutations.NumMaxConnections);
-			});
+			})
+		.EmplaceStep<FLeaveSessionHelper>(MoveTemp(LeaveSessionHelperParams));
 
 	RunToCompletion();
 }
diff --git a/Engine/Source/Programs/Online/OnlineTestsCore/Source/Private/Tests/Social/AcceptFriendInviteTests.cpp b/Engine/Source/Programs/Online/OnlineTestsCore/Source/Private/Tests/Social/AcceptFriendInviteTests.cpp
--- a/Engine/Source/Programs/Online/OnlineTestsCore/Source/Private/Tests/Social/AcceptFriendInviteTests.cpp
+++ b/Engine/Source/Programs/Online/OnlineTestsCore/Source/Private/Tests/Social/AcceptFriendInviteTests.cpp
@@ -2,30 +2,27 @@
 
 #include "Helpers/Social/AcceptFriendInviteHelper.h"
 #include "Helpers/Auth/AuthLogout.h"
-
-#include "OnlineCatchHelper.h"
+#include "Online/AuthCommon.h"
 
 #define SOCIAL_TAG "[suite_social]"
 #define EG_SOCIAL_ACCEPTFRIENDINVITE_TAG SOCIAL_TAG "[acceptfriendinvite]"
 #define EG_SOCIAL_ACCEPTFRIENDINVITEEOS_TAG SOCIAL_TAG "[acceptfriendinvite][.EOS]"
-#define EG_SOCIAL_DISABLED_TAG SOCIAL_TAG "[socialdisabled]"
+
 #define SOCIAL_TEST_CASE(x, ...) ONLINE_TEST_CASE(x, SOCIAL_TAG __VA_ARGS__)
 
-//SOCIAL_TEST_CASE("Verify that AcceptFriendInvite returns an error if call with an invalid local user account id", EG_SOCIAL_ACCEPTFRIENDINVITE_TAG)
-//{
-//	const int32 NumUsersToLogin = 0;
-//
-//	FAcceptFriendInvite::Params OpAcceptFriendInviteParams;
-//	FAcceptFriendInviteHelper::FHelperParams AcceptInviteHelperParams;
-//	AcceptInviteHelperParams.OpParams = &OpAcceptFriendInviteParams;
-//	AcceptInviteHelperParams.OpParams->LocalAccountId = FAccountId();
-//	AcceptInviteHelperParams.ExpectedError = TOnlineResult<FAcceptFriendInvite>(Errors::InvalidParams());
-//
-//	GetLoginPipeline(NumUsersToLogin)
-//		.EmplaceStep<FAcceptFriendInviteHelper>(MoveTemp(AcceptInviteHelperParams));
-//
-//	RunToCompletion();
-//}
+SOCIAL_TEST_CASE("Verify that AcceptFriendInvite returns an error if call with an invalid local user account id", EG_SOCIAL_ACCEPTFRIENDINVITE_TAG)
+{
+	FAcceptFriendInvite::Params OpAcceptFriendInviteParams;
+	FAcceptFriendInviteHelper::FHelperParams AcceptInviteHelperParams;
+	AcceptInviteHelperParams.OpParams = &OpAcceptFriendInviteParams;
+	AcceptInviteHelperParams.OpParams->LocalAccountId = FAccountId();
+	AcceptInviteHelperParams.ExpectedError = TOnlineResult<FAcceptFriendInvite>(Errors::InvalidParams());
+
+	GetPipeline()
+		.EmplaceStep<FAcceptFriendInviteHelper>(MoveTemp(AcceptInviteHelperParams));
+
+	RunToCompletion();
+}
 
 SOCIAL_TEST_CASE("Verify that AcceptFriendInvite returns an error if call with an target user account id", EG_SOCIAL_ACCEPTFRIENDINVITE_TAG)
 {
@@ -37,7 +34,7 @@ SOCIAL_TEST_CASE("Verify that AcceptFriendInvite returns an error if call with a
 	AcceptInviteHelperParams.OpParams->TargetAccountId = FAccountId();
 	AcceptInviteHelperParams.ExpectedError = TOnlineResult<FAcceptFriendInvite>(Errors::InvalidParams());
 
-	FTestPipeline& LoginPipeline = GetLoginPipeline(AccountId);
+	FTestPipeline& LoginPipeline = GetLoginPipeline({ AccountId });
 
 	AcceptInviteHelperParams.OpParams->LocalAccountId = AccountId;
 
@@ -47,28 +44,50 @@ SOCIAL_TEST_CASE("Verify that AcceptFriendInvite returns an error if call with a
 	RunToCompletion();
 }
 
-//SOCIAL_TEST_CASE("Verify that AcceptFriendInvite returns a fail message if the local user is not logged in", EG_SOCIAL_ACCEPTFRIENDINVITEEOS_TAG)
-//{
-//	FAccountId FirstAccountId, SecondAccountId;
-//	int32 UserNumToLogout = 1;
-//	bool bLogout = true;
-//
-//	FAcceptFriendInvite::Params OpAcceptFriendInviteParams;
-//	FAcceptFriendInviteHelper::FHelperParams AcceptInviteHelperParams;
-//	AcceptInviteHelperParams.OpParams = &OpAcceptFriendInviteParams;
-//	AcceptInviteHelperParams.ExpectedError = TOnlineResult<FAcceptFriendInvite>(Errors::NotLoggedIn());
-//
-//	FTestPipeline& LoginPipeline = GetLoginPipeline(FirstAccountId, SecondAccountId);
-//
-//	AcceptInviteHelperParams.OpParams->LocalAccountId = FirstAccountId;
-//	AcceptInviteHelperParams.OpParams->TargetAccountId = SecondAccountId;
-//
-//	LoginPipeline
-//		.EmplaceStep<FAuthLogoutStep>(FPlatformMisc::GetPlatformUserForUserIndex(0))
-//		.EmplaceStep<FAcceptFriendInviteHelper>(MoveTemp(AcceptInviteHelperParams));
-//
-//	RunToCompletion(bLogout, UserNumToLogout);
-//}
+SOCIAL_TEST_CASE("Verify that AcceptFriendInvite returns a fail message if the local user is not logged in", EG_SOCIAL_ACCEPTFRIENDINVITEEOS_TAG)
+{
+	FAccountId FirstAccountId, SecondAccountId;
+
+	int32 UserNumToLogin = 1;
+	TSharedPtr<FPlatformUserId> FirstAccountPlatformUserId = MakeShared<FPlatformUserId>();
+	TSharedPtr<FPlatformUserId> SecondAccountPlatformUserId = MakeShared<FPlatformUserId>();
+
+	bool bLogout = false;
+
+	FAcceptFriendInvite::Params OpAcceptFriendInviteParams;
+	FAcceptFriendInviteHelper::FHelperParams AcceptInviteHelperParams;
+	AcceptInviteHelperParams.OpParams = &OpAcceptFriendInviteParams;
+	AcceptInviteHelperParams.ExpectedError = TOnlineResult<FAcceptFriendInvite>(Errors::NotLoggedIn());
+
+	FTestPipeline& LoginPipeline = GetLoginPipeline(UserNumToLogin, { FirstAccountId, SecondAccountId });
+
+	AcceptInviteHelperParams.OpParams->LocalAccountId = FirstAccountId;
+	AcceptInviteHelperParams.OpParams->TargetAccountId = SecondAccountId;
+
+	LoginPipeline
+		.EmplaceLambda([&FirstAccountId, &SecondAccountId, FirstAccountPlatformUserId, SecondAccountPlatformUserId](SubsystemType OnlineSubsystem)
+			{
+				UE::Online::IAuthPtr OnlineAuthPtr = OnlineSubsystem->GetAuthInterface();
+				REQUIRE(OnlineAuthPtr);
+
+				UE::Online::TOnlineResult<UE::Online::FAuthGetLocalOnlineUserByOnlineAccountId> FirstUserPlatfromUserIdResult = OnlineAuthPtr->GetLocalOnlineUserByOnlineAccountId({ FirstAccountId });
+				UE::Online::TOnlineResult<UE::Online::FAuthGetLocalOnlineUserByOnlineAccountId> SecondUserPlatfromUserIdResult = OnlineAuthPtr->GetLocalOnlineUserByOnlineAccountId({ SecondAccountId });
+
+				REQUIRE(FirstUserPlatfromUserIdResult.IsOk());
+				REQUIRE(SecondUserPlatfromUserIdResult.IsOk());
+
+				CHECK(FirstUserPlatfromUserIdResult.TryGetOkValue() != nullptr);
+				CHECK(SecondUserPlatfromUserIdResult.TryGetOkValue() != nullptr);
+
+				*FirstAccountPlatformUserId = FirstUserPlatfromUserIdResult.TryGetOkValue()->AccountInfo->PlatformUserId;
+				*SecondAccountPlatformUserId = SecondUserPlatfromUserIdResult.TryGetOkValue()->AccountInfo->PlatformUserId;
+			})
+		.EmplaceStep<FAuthLogoutStep>(MoveTemp(FirstAccountPlatformUserId))
+		.EmplaceStep<FAcceptFriendInviteHelper>(MoveTemp(AcceptInviteHelperParams))
+		.EmplaceStep<FAuthLogoutStep>(MoveTemp(SecondAccountPlatformUserId));
+
+	RunToCompletion(bLogout);
+}
 
 //SOCIAL_TEST_CASE("Verify that AcceptFriendInvite returns a fail message if ERelationship with target user is Friend")
 //{
@@ -90,7 +109,7 @@ SOCIAL_TEST_CASE("Verify that AcceptFriendInvite returns an error if call with a
 //	// TODO
 //}
 
-//SOCIAL_TEST_CASE("Verify that AcceptFriendInvite returns a fail message if ERelationship with target user is Blocked", EG_SOCIAL_DISABLED_TAG)
+//SOCIAL_TEST_CASE("Verify that AcceptFriendInvite returns a fail message if ERelationship with target user is Blocked")
 //{
 //	// TODO
 //}
diff --git a/Engine/Source/Programs/Online/OnlineTestsCore/Source/Private/Tests/Social/BlockUserTests.cpp b/Engine/Source/Programs/Online/OnlineTestsCore/Source/Private/Tests/Social/BlockUserTests.cpp
--- a/Engine/Source/Programs/Online/OnlineTestsCore/Source/Private/Tests/Social/BlockUserTests.cpp
+++ b/Engine/Source/Programs/Online/OnlineTestsCore/Source/Private/Tests/Social/BlockUserTests.cpp
@@ -3,39 +3,22 @@
 #include "Helpers/Social/BlockUserHelper.h"
 #include "Helpers/Auth/AuthLogout.h"
 
-#include "OnlineCatchHelper.h"
-
 #define SOCIAL_TAG "[suite_social]"
 #define EG_SOCIAL_BLOCKUSER_TAG SOCIAL_TAG "[blockuser]"
 #define EG_SOCIAL_BLOCKUSEREOS_TAG SOCIAL_TAG "[blockuser][.EOS]"
+
 #define SOCIAL_TEST_CASE(x, ...) ONLINE_TEST_CASE(x, SOCIAL_TAG __VA_ARGS__)
 
 SOCIAL_TEST_CASE("Verify that BlockUser returns a error if call with an invalid local user account id", EG_SOCIAL_BLOCKUSER_TAG)
 {
-	const int32 NumUsersToLogin = 0;
-
 	FBlockUser::Params OpBlockUserParams;
 	FBlockUserHelper::FHelperParams BlockUserHelperHelperParams;
 	BlockUserHelperHelperParams.OpParams = &OpBlockUserParams;
 	BlockUserHelperHelperParams.OpParams->LocalAccountId = FAccountId();
-
-	SubsystemType OnlineSubsystem = GetSubsystem();
-	EOnlineServices ServicesProvider = OnlineSubsystem->GetServicesProvider();
-	
-	/*if (ServicesProvider == EOnlineServices::Epic)
-	{
-		BlockUserHelperHelperParams.ExpectedError = TOnlineResult<FBlockUser>(Errors::InvalidParams());
-
-	}
-	else if (ServicesProvider == EOnlineServices::Xbox)
-	{
-		BlockUserHelperHelperParams.ExpectedError = TOnlineResult<FBlockUser>(Errors::InvalidUser());
-	}*/
-
 	BlockUserHelperHelperParams.ExpectedError = TOnlineResult<FBlockUser>(Errors::InvalidParams());
 
 
-	GetLoginPipeline(NumUsersToLogin)
+	GetPipeline()
 		.EmplaceStep<FBlockUserHelper>(MoveTemp(BlockUserHelperHelperParams));
 
 	RunToCompletion();
@@ -51,7 +34,7 @@ SOCIAL_TEST_CASE("Verify that BlockUser returns a error if call with an target u
 	BlockUserHelperHelperParams.OpParams->TargetAccountId = FAccountId();
 	BlockUserHelperHelperParams.ExpectedError = TOnlineResult<FBlockUser>(Errors::InvalidParams());
 
-	FTestPipeline& LoginPipeline = GetLoginPipeline(AccountId);
+	FTestPipeline& LoginPipeline = GetLoginPipeline({ AccountId });
 
 	BlockUserHelperHelperParams.OpParams->LocalAccountId = AccountId;
 
@@ -64,24 +47,46 @@ SOCIAL_TEST_CASE("Verify that BlockUser returns a error if call with an target u
 SOCIAL_TEST_CASE("Verify that BlockUser returns a fail message if the local user is not logged in", EG_SOCIAL_BLOCKUSEREOS_TAG)
 {
 	FAccountId FirstAccountId, SecondAccountId;
-	int32 UserNumToLogout = 1;
-	bool bLogout = true;
+	
+	int32 UserNumToLogin = 1;
+	TSharedPtr<FPlatformUserId> FirstAccountPlatformUserId = MakeShared<FPlatformUserId>();
+	TSharedPtr<FPlatformUserId> SecondAccountPlatformUserId = MakeShared<FPlatformUserId>();
+
+	bool bLogout = false;
 
 	FBlockUser::Params OpBlockUserParams;
 	FBlockUserHelper::FHelperParams BlockUserHelperHelperParams;
 	BlockUserHelperHelperParams.OpParams = &OpBlockUserParams;
 	BlockUserHelperHelperParams.ExpectedError = TOnlineResult<FBlockUser>(Errors::NotLoggedIn());
 
-	FTestPipeline& LoginPipeline = GetLoginPipeline(FirstAccountId, SecondAccountId);
+	FTestPipeline& LoginPipeline = GetLoginPipeline(UserNumToLogin, { FirstAccountId, SecondAccountId });
 
 	BlockUserHelperHelperParams.OpParams->LocalAccountId = FirstAccountId;
 	BlockUserHelperHelperParams.OpParams->TargetAccountId = SecondAccountId;
 
 	LoginPipeline
-		.EmplaceStep<FAuthLogoutStep>(FPlatformMisc::GetPlatformUserForUserIndex(0))
-		.EmplaceStep<FBlockUserHelper>(MoveTemp(BlockUserHelperHelperParams));
+		.EmplaceLambda([&FirstAccountId, &SecondAccountId, FirstAccountPlatformUserId, SecondAccountPlatformUserId](SubsystemType OnlineSubsystem)
+			{
+				UE::Online::IAuthPtr OnlineAuthPtr = OnlineSubsystem->GetAuthInterface();
+				REQUIRE(OnlineAuthPtr);
+
+				UE::Online::TOnlineResult<UE::Online::FAuthGetLocalOnlineUserByOnlineAccountId> FirstUserPlatfromUserIdResult = OnlineAuthPtr->GetLocalOnlineUserByOnlineAccountId({ FirstAccountId });
+				UE::Online::TOnlineResult<UE::Online::FAuthGetLocalOnlineUserByOnlineAccountId> SecondUserPlatfromUserIdResult = OnlineAuthPtr->GetLocalOnlineUserByOnlineAccountId({ SecondAccountId });
+
+				REQUIRE(FirstUserPlatfromUserIdResult.IsOk());
+				REQUIRE(SecondUserPlatfromUserIdResult.IsOk());
+
+				CHECK(FirstUserPlatfromUserIdResult.TryGetOkValue() != nullptr);
+				CHECK(SecondUserPlatfromUserIdResult.TryGetOkValue() != nullptr);
+
+				*FirstAccountPlatformUserId = FirstUserPlatfromUserIdResult.TryGetOkValue()->AccountInfo->PlatformUserId;
+				*SecondAccountPlatformUserId = SecondUserPlatfromUserIdResult.TryGetOkValue()->AccountInfo->PlatformUserId;
+			})
+		.EmplaceStep<FAuthLogoutStep>(MoveTemp(FirstAccountPlatformUserId))
+		.EmplaceStep<FBlockUserHelper>(MoveTemp(BlockUserHelperHelperParams))
+		.EmplaceStep<FAuthLogoutStep>(MoveTemp(SecondAccountPlatformUserId));
 
-	RunToCompletion(bLogout, UserNumToLogout);
+	RunToCompletion(bLogout);
 }
 
 SOCIAL_TEST_CASE("Verify that BlockUser completes successfully if both users are logged in", EG_SOCIAL_BLOCKUSEREOS_TAG)
@@ -92,7 +97,7 @@ SOCIAL_TEST_CASE("Verify that BlockUser completes successfully if both users are
 	FBlockUserHelper::FHelperParams BlockUserHelperHelperParams;
 	BlockUserHelperHelperParams.OpParams = &OpBlockUserParams;
 
-	FTestPipeline& LoginPipeline = GetLoginPipeline(FirstAccountId, SecondAccountId);
+	FTestPipeline& LoginPipeline = GetLoginPipeline({ FirstAccountId, SecondAccountId });
 
 	BlockUserHelperHelperParams.OpParams->LocalAccountId = FirstAccountId;
 	BlockUserHelperHelperParams.OpParams->TargetAccountId = SecondAccountId;
diff --git a/Engine/Source/Programs/Online/OnlineTestsCore/Source/Private/Tests/Social/GetBlockedUsersTests.cpp b/Engine/Source/Programs/Online/OnlineTestsCore/Source/Private/Tests/Social/GetBlockedUsersTests.cpp
--- a/Engine/Source/Programs/Online/OnlineTestsCore/Source/Private/Tests/Social/GetBlockedUsersTests.cpp
+++ b/Engine/Source/Programs/Online/OnlineTestsCore/Source/Private/Tests/Social/GetBlockedUsersTests.cpp
@@ -3,16 +3,13 @@
 #include "Helpers/Social/GetBlockedUsersHelper.h"
 #include "Helpers/Auth/AuthLogout.h"
 
-#include "OnlineCatchHelper.h"
-
 #define SOCIAL_TAG "[suite_social]"
 #define EG_SOCIAL_GETBLOCKEDUSERSGDK_TAG SOCIAL_TAG "[getblockedusers][.GDK]"
-#define EG_SOCIAL_DISABLED_TAG SOCIAL_TAG "[socialdisabled]"
+
 #define SOCIAL_TEST_CASE(x, ...) ONLINE_TEST_CASE(x, SOCIAL_TAG __VA_ARGS__)
 
 SOCIAL_TEST_CASE("Verify that GetBlockedUsers returns a error if call with an invalid local user account id", EG_SOCIAL_GETBLOCKEDUSERSGDK_TAG)
 {
-	const int32 NumUsersToLogin = 0;
 	FGetBlockedUsers::Params OpGetBlockedUsersParams;
 	FGetBlockedUsersHelper::FHelperParams GetBlockedUsersHelperParams;
 	GetBlockedUsersHelperParams.OpParams = &OpGetBlockedUsersParams;
@@ -31,7 +28,7 @@ SOCIAL_TEST_CASE("Verify that GetBlockedUsers returns a error if call with an in
 		GetBlockedUsersHelperParams.ExpectedError = TOnlineResult<FGetBlockedUsers>(Errors::InvalidUser());
 	}
 
-	GetLoginPipeline(NumUsersToLogin)
+	GetPipeline()
 		.EmplaceStep<FGetBlockedUsersHelper>(MoveTemp(GetBlockedUsersHelperParams));
 
 	RunToCompletion();
@@ -58,7 +55,7 @@ SOCIAL_TEST_CASE("Verify that GetBlockedUsers returns a fail message if there ar
 		GetBlockedUsersHelperParams.ExpectedError = TOnlineResult<FGetBlockedUsers>(Errors::InvalidUser());
 	}
 
-	FTestPipeline& LoginPipeline = GetLoginPipeline(AccountId);
+	FTestPipeline& LoginPipeline = GetLoginPipeline({ AccountId });
 
 	GetBlockedUsersHelperParams.OpParams->LocalAccountId = AccountId;
 
diff --git a/Engine/Source/Programs/Online/OnlineTestsCore/Source/Private/Tests/Social/GetFriendsTests.cpp b/Engine/Source/Programs/Online/OnlineTestsCore/Source/Private/Tests/Social/GetFriendsTests.cpp
--- a/Engine/Source/Programs/Online/OnlineTestsCore/Source/Private/Tests/Social/GetFriendsTests.cpp
+++ b/Engine/Source/Programs/Online/OnlineTestsCore/Source/Private/Tests/Social/GetFriendsTests.cpp
@@ -4,8 +4,6 @@
 #include "Helpers/Social/GetFriendsHelper.h"
 #include "Helpers/Auth/AuthLogout.h"
 
-#include "OnlineCatchHelper.h"
-
 #define SOCIAL_TAG "[suite_social]"
 #define EG_SOCIAL_GETFRIENDSEOS_TAG SOCIAL_TAG "[getfriends][.EOS]"
 #define EG_SOCIAL_GETFRIENDS_TAG SOCIAL_TAG "[getfriends]"
@@ -15,32 +13,29 @@
 SOCIAL_TEST_CASE("Verify that GetFriends returns an empty list if there are no cached Friends", EG_SOCIAL_GETFRIENDSEOS_TAG)
 {
 	FAccountId AccountId;	
+
 	int32 UserNumToLogin = 5;
-	int32 UserNumToLogout = 5;
 
 	FGetFriends::Params OpGetParams;
 	FGetFriendsHelper::FHelperParams GetFriendsHelperParams;
 	GetFriendsHelperParams.OpParams = &OpGetParams;
 	GetFriendsHelperParams.ExpectedError = TOnlineResult<FGetFriends>(Errors::InvalidState());
 
-	FTestPipeline& LoginPipeline = GetLoginPipeline(UserNumToLogin, AccountId);
+	FTestPipeline& LoginPipeline = GetLoginPipeline(UserNumToLogin, { AccountId });
 
 	GetFriendsHelperParams.OpParams->LocalAccountId = AccountId;
 
-	bool bLogout = true;
-
 	LoginPipeline
 		.EmplaceStep<FGetFriendsHelper>(MoveTemp(GetFriendsHelperParams));
 
-	RunToCompletion(bLogout, UserNumToLogout);
+	RunToCompletion();
 }
 
 SOCIAL_TEST_CASE("Verify that GetFriends returns a list of 1 Friend if there is 1 cached Friend", EG_SOCIAL_GETFRIENDSEOS_TAG)
 {
 	FAccountId AccountId;
+	
 	int32 UserNumToLogin = 6;
-	int32 UserNumToLogout = 6;
-	bool bLogout = true;
 
 	FQueryFriends::Params OpQueryParams;
 	FQueryFriendsHelper::FHelperParams QueryFriendsHelperParams;
@@ -50,7 +45,7 @@ SOCIAL_TEST_CASE("Verify that GetFriends returns a list of 1 Friend if there is
 	FGetFriendsHelper::FHelperParams GetFriendsHelperParams;
 	GetFriendsHelperParams.OpParams = &OpGetParams;
 
-	FTestPipeline& LoginPipeline = GetLoginPipeline(UserNumToLogin, AccountId);
+	FTestPipeline& LoginPipeline = GetLoginPipeline(UserNumToLogin, { AccountId });
 
 	QueryFriendsHelperParams.OpParams->LocalAccountId = AccountId;
 	GetFriendsHelperParams.OpParams->LocalAccountId = AccountId;
@@ -61,7 +56,7 @@ SOCIAL_TEST_CASE("Verify that GetFriends returns a list of 1 Friend if there is
 		.EmplaceStep<FQueryFriendsHelper>(MoveTemp(QueryFriendsHelperParams))
 		.EmplaceStep<FGetFriendsHelper>(MoveTemp(GetFriendsHelperParams), ExpectedFriendsNum);
 
-	RunToCompletion(bLogout, UserNumToLogout);
+	RunToCompletion();
 }
 
 SOCIAL_TEST_CASE("Verify that GetFriends returns a list of all cached Friends if there are multiple cached Friends", EG_SOCIAL_GETFRIENDSEOS_TAG)
@@ -76,7 +71,7 @@ SOCIAL_TEST_CASE("Verify that GetFriends returns a list of all cached Friends if
 	FGetFriendsHelper::FHelperParams GetFriendsHelperParams;
 	GetFriendsHelperParams.OpParams = &OpGetParams;
 
-	FTestPipeline& LoginPipeline = GetLoginPipeline(AccountId);
+	FTestPipeline& LoginPipeline = GetLoginPipeline({ AccountId });
 
 	QueryFriendsHelperParams.OpParams->LocalAccountId = AccountId;
 	GetFriendsHelperParams.OpParams->LocalAccountId = AccountId;
@@ -92,23 +87,21 @@ SOCIAL_TEST_CASE("Verify that GetFriends returns a list of all cached Friends if
 
 SOCIAL_TEST_CASE("Verify that GetFriend returns a fail message if there are no cached Friends", EG_SOCIAL_GETFRIENDSEOS_TAG)
 {
-	FAccountId AccountId;	
+	FAccountId AccountId;
+	
 	int32 UserNumToLogin = 5;
-	int32 UserNumToLogout = 5;
 
 	FGetFriends::Params OpGetParams;
 	FGetFriendsHelper::FHelperParams GetFriendsHelperParams;
 	GetFriendsHelperParams.OpParams = &OpGetParams;
 	GetFriendsHelperParams.ExpectedError = TOnlineResult<FGetFriends>(Errors::InvalidState());
 
-	FTestPipeline& LoginPipeline = GetLoginPipeline(UserNumToLogin, AccountId);
+	FTestPipeline& LoginPipeline = GetLoginPipeline(UserNumToLogin, { AccountId });
 
 	GetFriendsHelperParams.OpParams->LocalAccountId = AccountId;
 
-	bool bLogout = true;
-
 	LoginPipeline
 		.EmplaceStep<FGetFriendsHelper>(MoveTemp(GetFriendsHelperParams));
 
-	RunToCompletion(bLogout, UserNumToLogout);
+	RunToCompletion();
 }
diff --git a/Engine/Source/Programs/Online/OnlineTestsCore/Source/Private/Tests/Social/QueryBlockedUsersTests.cpp b/Engine/Source/Programs/Online/OnlineTestsCore/Source/Private/Tests/Social/QueryBlockedUsersTests.cpp
--- a/Engine/Source/Programs/Online/OnlineTestsCore/Source/Private/Tests/Social/QueryBlockedUsersTests.cpp
+++ b/Engine/Source/Programs/Online/OnlineTestsCore/Source/Private/Tests/Social/QueryBlockedUsersTests.cpp
@@ -3,17 +3,13 @@
 #include "Helpers/Social/QueryBlockedUsersHelper.h"
 #include "Helpers/Auth/AuthLogout.h"
 
-#include "OnlineCatchHelper.h"
-
 #define SOCIAL_TAG "[suite_social]"
 #define EG_SOCIAL_QUERYBLOCKEDUSERSGDK_TAG SOCIAL_TAG "[queryblockedusers][.GDK]"
-#define EG_SOCIAL_DISABLED_TAG SOCIAL_TAG "[socialdisabled]"
+
 #define SOCIAL_TEST_CASE(x, ...) ONLINE_TEST_CASE(x, SOCIAL_TAG __VA_ARGS__)
 
 SOCIAL_TEST_CASE("Verify that QueryBlockedUsers returns a error if call with an invalid local user account id", EG_SOCIAL_QUERYBLOCKEDUSERSGDK_TAG)
 {
-	const int32 NumUsersToLogin = 0;
-	
 	FQueryBlockedUsers::Params OpQueryParams;
 	FQueryBlockedUsersHelper::FHelperParams QueryBlockedUsersParams;
 	QueryBlockedUsersParams.OpParams = &OpQueryParams;
@@ -32,7 +28,7 @@ SOCIAL_TEST_CASE("Verify that QueryBlockedUsers returns a error if call with an
 		QueryBlockedUsersParams.ExpectedError = TOnlineResult<FQueryBlockedUsers>(Errors::InvalidUser());
 	}
 	
-	GetLoginPipeline(NumUsersToLogin)
+	GetPipeline()
 		.EmplaceStep<FQueryBlockedUsersHelper>(MoveTemp(QueryBlockedUsersParams));
 	
 	RunToCompletion();
diff --git a/Engine/Source/Programs/Online/OnlineTestsCore/Source/Private/Tests/Social/QueryFriendsTests.cpp b/Engine/Source/Programs/Online/OnlineTestsCore/Source/Private/Tests/Social/QueryFriendsTests.cpp
--- a/Engine/Source/Programs/Online/OnlineTestsCore/Source/Private/Tests/Social/QueryFriendsTests.cpp
+++ b/Engine/Source/Programs/Online/OnlineTestsCore/Source/Private/Tests/Social/QueryFriendsTests.cpp
@@ -4,30 +4,44 @@
 #include "Helpers/Social/GetFriendsHelper.h"
 #include "Helpers/Auth/AuthLogout.h"
 
-#include "OnlineCatchHelper.h"
-
 #define SOCIAL_TAG "[suite_social]"
 #define EG_SOCIAL_QUERYFRIENDS_TAG SOCIAL_TAG "[queryfriends]"
 #define EG_SOCIAL_QUERYFRIENDSEOS_TAG SOCIAL_TAG "[queryfriends][.EOS]"
+
 #define SOCIAL_TEST_CASE(x, ...) ONLINE_TEST_CASE(x, SOCIAL_TAG __VA_ARGS__)
 
 SOCIAL_TEST_CASE("Verify that QueryFriends returns a fail message if the local user is not logged in", EG_SOCIAL_QUERYFRIENDSEOS_TAG)
 {
 	FAccountId AccountId;
 
+	int32 UserNumToLogin = 1;
+	TSharedPtr<FPlatformUserId> AccountPlatformUserId = MakeShared<FPlatformUserId>();
+	bool bLogout = false;
+
 	FQueryFriends::Params OpQueryParams;
 	FQueryFriendsHelper::FHelperParams QueryFriendsHelperParams;
 	QueryFriendsHelperParams.OpParams = &OpQueryParams;
 	QueryFriendsHelperParams.ExpectedError = TOnlineResult<FQueryFriends>(Errors::NotLoggedIn());
 
-	FTestPipeline& LoginPipeline = GetLoginPipeline(AccountId);
+	FTestPipeline& LoginPipeline = GetLoginPipeline(UserNumToLogin, { AccountId });
 
 	QueryFriendsHelperParams.OpParams->LocalAccountId = AccountId;
 
-	bool bLogout = false;
 
 	LoginPipeline
-		.EmplaceStep<FAuthLogoutStep>(FPlatformMisc::GetPlatformUserForUserIndex(0))
+		.EmplaceLambda([&AccountId, AccountPlatformUserId](SubsystemType OnlineSubsystem)
+			{
+				UE::Online::IAuthPtr OnlineAuthPtr = OnlineSubsystem->GetAuthInterface();
+				REQUIRE(OnlineAuthPtr);
+
+				UE::Online::TOnlineResult<UE::Online::FAuthGetLocalOnlineUserByOnlineAccountId> UserPlatfromUserIdResult = OnlineAuthPtr->GetLocalOnlineUserByOnlineAccountId({ AccountId });
+				
+				REQUIRE(UserPlatfromUserIdResult.IsOk());
+				CHECK(UserPlatfromUserIdResult.TryGetOkValue() != nullptr);
+
+				*AccountPlatformUserId = UserPlatfromUserIdResult.TryGetOkValue()->AccountInfo->PlatformUserId;
+			})
+		.EmplaceStep<FAuthLogoutStep>(MoveTemp(AccountPlatformUserId))
 		.EmplaceStep<FQueryFriendsHelper>(MoveTemp(QueryFriendsHelperParams));
 
 	RunToCompletion(bLogout);
@@ -35,15 +49,13 @@ SOCIAL_TEST_CASE("Verify that QueryFriends returns a fail message if the local u
 
 SOCIAL_TEST_CASE("Verify that QueryFriends returns a error if call with an invalid account id", EG_SOCIAL_QUERYFRIENDS_TAG)
 {
-	const int32 NumUsersToLogin = 0;
-
 	FQueryFriends::Params OpQueryParams;
 	FQueryFriendsHelper::FHelperParams QueryFriendsHelperParams;
 	QueryFriendsHelperParams.OpParams = &OpQueryParams;
 	QueryFriendsHelperParams.OpParams->LocalAccountId = FAccountId();
 	QueryFriendsHelperParams.ExpectedError = TOnlineResult<FQueryFriends>(Errors::InvalidParams());
 
-	GetLoginPipeline(NumUsersToLogin)
+	GetPipeline()
 		.EmplaceStep<FQueryFriendsHelper>(MoveTemp(QueryFriendsHelperParams));
 
 	RunToCompletion();
@@ -53,8 +65,6 @@ SOCIAL_TEST_CASE("Verify that QueryFriends caches no Friends if no Friends exist
 {
 	FAccountId AccountId;
 	int32 UserNumToLogin = 5;
-	int32 UserNumToLogout = 5;
-	bool bLogout = true;
 
 	FQueryFriends::Params OpQueryParams;
 	FQueryFriendsHelper::FHelperParams QueryFriendsHelperParams;
@@ -64,7 +74,7 @@ SOCIAL_TEST_CASE("Verify that QueryFriends caches no Friends if no Friends exist
 	FGetFriendsHelper::FHelperParams GetFriendsHelperParams;
 	GetFriendsHelperParams.OpParams = &OpGetParams;
 
-	FTestPipeline& LoginPipeline = GetLoginPipeline(UserNumToLogin, AccountId);
+	FTestPipeline& LoginPipeline = GetLoginPipeline(UserNumToLogin, { AccountId });
 
 	QueryFriendsHelperParams.OpParams->LocalAccountId = AccountId;
 	GetFriendsHelperParams.OpParams->LocalAccountId = AccountId;
@@ -73,15 +83,13 @@ SOCIAL_TEST_CASE("Verify that QueryFriends caches no Friends if no Friends exist
 		.EmplaceStep<FQueryFriendsHelper>(MoveTemp(QueryFriendsHelperParams))
 		.EmplaceStep<FGetFriendsHelper>(MoveTemp(GetFriendsHelperParams));
 
-	RunToCompletion(bLogout, UserNumToLogout);
+	RunToCompletion();
 }
 
 SOCIAL_TEST_CASE("Verify that QueryFriends caches one Friend if only one Friend exists for this user", EG_SOCIAL_QUERYFRIENDSEOS_TAG)
 {
 	FAccountId AccountId;
 	int32 UserNumToLogin = 6;
-	int32 UserNumToLogout = 6;
-	bool bLogout = true;
 
 	FQueryFriends::Params OpQueryParams;
 	FQueryFriendsHelper::FHelperParams QueryFriendsHelperParams;
@@ -91,7 +99,7 @@ SOCIAL_TEST_CASE("Verify that QueryFriends caches one Friend if only one Friend
 	FGetFriendsHelper::FHelperParams GetFriendsHelperParams;
 	GetFriendsHelperParams.OpParams = &OpGetParams;
 		
-	FTestPipeline& LoginPipeline = GetLoginPipeline(UserNumToLogin, AccountId);
+	FTestPipeline& LoginPipeline = GetLoginPipeline(UserNumToLogin, { AccountId });
 	
 	QueryFriendsHelperParams.OpParams->LocalAccountId = AccountId;
 	GetFriendsHelperParams.OpParams->LocalAccountId = AccountId;
@@ -102,7 +110,7 @@ SOCIAL_TEST_CASE("Verify that QueryFriends caches one Friend if only one Friend
 		.EmplaceStep<FQueryFriendsHelper>(MoveTemp(QueryFriendsHelperParams))
 		.EmplaceStep<FGetFriendsHelper>(MoveTemp(GetFriendsHelperParams), ExpectedFriendsNum);
 	
-	RunToCompletion(bLogout, UserNumToLogout);
+	RunToCompletion();
 }
 
 SOCIAL_TEST_CASE("Verify that QueryFriends caches all Friends if multiple Friends exist for this user", EG_SOCIAL_QUERYFRIENDSEOS_TAG)
@@ -118,7 +126,7 @@ SOCIAL_TEST_CASE("Verify that QueryFriends caches all Friends if multiple Friend
 	FGetFriendsHelper::FHelperParams GetFriendsHelperParams;
 	GetFriendsHelperParams.OpParams = &OpGetParams;
 
-	FTestPipeline& LoginPipeline = GetLoginPipeline(AccountId);
+	FTestPipeline& LoginPipeline = GetLoginPipeline({ AccountId });
 
 	QueryFriendsHelperParams.OpParams->LocalAccountId = AccountId;
 	GetFriendsHelperParams.OpParams->LocalAccountId = AccountId;
diff --git a/Engine/Source/Programs/Online/OnlineTestsCore/Source/Private/Tests/Social/RejectFriendInviteTests.cpp b/Engine/Source/Programs/Online/OnlineTestsCore/Source/Private/Tests/Social/RejectFriendInviteTests.cpp
--- a/Engine/Source/Programs/Online/OnlineTestsCore/Source/Private/Tests/Social/RejectFriendInviteTests.cpp
+++ b/Engine/Source/Programs/Online/OnlineTestsCore/Source/Private/Tests/Social/RejectFriendInviteTests.cpp
@@ -3,25 +3,21 @@
 #include "Helpers/Social/RejectFriendInviteHelper.h"
 #include "Helpers/Auth/AuthLogout.h"
 
-#include "OnlineCatchHelper.h"
-
 #define SOCIAL_TAG "[suite_social]"
 #define EG_SOCIAL_REJECTFRIENDINVITE_TAG SOCIAL_TAG "[rejectfriendinvite]"
 #define EG_SOCIAL_REJECTFRIENDINVITEEOS_TAG SOCIAL_TAG "[queryfriends][.EOS]"
-#define EG_SOCIAL_DISABLED_TAG SOCIAL_TAG "[socialdisabled]"
+
 #define SOCIAL_TEST_CASE(x, ...) ONLINE_TEST_CASE(x, SOCIAL_TAG __VA_ARGS__)
 
 SOCIAL_TEST_CASE("Verify that RejectFriendInvite returns a error if call with an invalid local user account id", EG_SOCIAL_REJECTFRIENDINVITE_TAG)
 {
-	const int32 NumUsersToLogin = 0;
-
 	FRejectFriendInvite::Params OpRejectFriendInviteParams;
 	FRejectFriendInviteHelper::FHelperParams RejectFriendInviteHelperParams;
 	RejectFriendInviteHelperParams.OpParams = &OpRejectFriendInviteParams;
 	RejectFriendInviteHelperParams.OpParams->LocalAccountId = FAccountId();
 	RejectFriendInviteHelperParams.ExpectedError = TOnlineResult<FRejectFriendInvite>(Errors::InvalidParams());
 
-	GetLoginPipeline(NumUsersToLogin)
+	GetPipeline()
 		.EmplaceStep<FRejectFriendInviteHelper>(MoveTemp(RejectFriendInviteHelperParams));
 
 	RunToCompletion();
@@ -37,7 +33,7 @@ SOCIAL_TEST_CASE("Verify that RejectFriendInvite returns a error if call with an
 	RejectFriendInviteHelperParams.OpParams->TargetAccountId = FAccountId();
 	RejectFriendInviteHelperParams.ExpectedError = TOnlineResult<FRejectFriendInvite>(Errors::InvalidParams());
 
-	FTestPipeline& LoginPipeline = GetLoginPipeline(AccountId);
+	FTestPipeline& LoginPipeline = GetLoginPipeline({ AccountId });
 
 	RejectFriendInviteHelperParams.OpParams->LocalAccountId = AccountId;
 
@@ -50,24 +46,45 @@ SOCIAL_TEST_CASE("Verify that RejectFriendInvite returns a error if call with an
 SOCIAL_TEST_CASE("Verify that RejectFriendInvite returns a fail message if the local user is not logged in", EG_SOCIAL_REJECTFRIENDINVITEEOS_TAG)
 {
 	FAccountId FirstAccountId, SecondAccountId;
-	int32 UserNumToLogout = 1;
-	bool bLogout = true;
+
+	int32 UserNumToLogin = 1;
+	TSharedPtr<FPlatformUserId> FirstAccountPlatformUserId = MakeShared<FPlatformUserId>();
+	TSharedPtr<FPlatformUserId> SecondAccountPlatformUserId = MakeShared<FPlatformUserId>();
+	bool bLogout = false;
 
 	FRejectFriendInvite::Params OpRejectFriendInviteParams;
 	FRejectFriendInviteHelper::FHelperParams RejectFriendInviteHelperParams;
 	RejectFriendInviteHelperParams.OpParams = &OpRejectFriendInviteParams;
 	RejectFriendInviteHelperParams.ExpectedError = TOnlineResult<FRejectFriendInvite>(Errors::NotLoggedIn());
 
-	FTestPipeline& LoginPipeline = GetLoginPipeline(FirstAccountId, SecondAccountId);
+	FTestPipeline& LoginPipeline = GetLoginPipeline(UserNumToLogin, { FirstAccountId, SecondAccountId });
 
 	RejectFriendInviteHelperParams.OpParams->LocalAccountId = FirstAccountId;
 	RejectFriendInviteHelperParams.OpParams->TargetAccountId = SecondAccountId;
 
 	LoginPipeline
-		.EmplaceStep<FAuthLogoutStep>(FPlatformMisc::GetPlatformUserForUserIndex(0))
-		.EmplaceStep<FRejectFriendInviteHelper>(MoveTemp(RejectFriendInviteHelperParams));
+		.EmplaceLambda([&FirstAccountId, &SecondAccountId, FirstAccountPlatformUserId, SecondAccountPlatformUserId](SubsystemType OnlineSubsystem)
+			{
+				UE::Online::IAuthPtr OnlineAuthPtr = OnlineSubsystem->GetAuthInterface();
+				REQUIRE(OnlineAuthPtr);
+
+				UE::Online::TOnlineResult<UE::Online::FAuthGetLocalOnlineUserByOnlineAccountId> FirstUserPlatfromUserIdResult = OnlineAuthPtr->GetLocalOnlineUserByOnlineAccountId({ FirstAccountId });
+				UE::Online::TOnlineResult<UE::Online::FAuthGetLocalOnlineUserByOnlineAccountId> SecondUserPlatfromUserIdResult = OnlineAuthPtr->GetLocalOnlineUserByOnlineAccountId({ SecondAccountId });
+
+				REQUIRE(FirstUserPlatfromUserIdResult.IsOk());
+				REQUIRE(SecondUserPlatfromUserIdResult.IsOk());
+
+				CHECK(FirstUserPlatfromUserIdResult.TryGetOkValue() != nullptr);
+				CHECK(SecondUserPlatfromUserIdResult.TryGetOkValue() != nullptr);
+
+				*FirstAccountPlatformUserId = FirstUserPlatfromUserIdResult.TryGetOkValue()->AccountInfo->PlatformUserId;
+				*SecondAccountPlatformUserId = SecondUserPlatfromUserIdResult.TryGetOkValue()->AccountInfo->PlatformUserId;
+			})
+		.EmplaceStep<FAuthLogoutStep>(MoveTemp(FirstAccountPlatformUserId))
+		.EmplaceStep<FRejectFriendInviteHelper>(MoveTemp(RejectFriendInviteHelperParams))
+		.EmplaceStep<FAuthLogoutStep>(MoveTemp(SecondAccountPlatformUserId));
 
-	RunToCompletion(bLogout, UserNumToLogout);
+	RunToCompletion(bLogout);
 }
 
 //SOCIAL_TEST_CASE("Verify that RejectFriendInvite returns a fail message if ERelationship with target user is Friend")
diff --git a/Engine/Source/Programs/Online/OnlineTestsCore/Source/Private/Tests/Social/SendFriendInviteTests.cpp b/Engine/Source/Programs/Online/OnlineTestsCore/Source/Private/Tests/Social/SendFriendInviteTests.cpp
--- a/Engine/Source/Programs/Online/OnlineTestsCore/Source/Private/Tests/Social/SendFriendInviteTests.cpp
+++ b/Engine/Source/Programs/Online/OnlineTestsCore/Source/Private/Tests/Social/SendFriendInviteTests.cpp
@@ -7,25 +7,21 @@
 #include "Helpers/Auth/AuthLogout.h"
 #include "Helpers/Auth/AuthLogin.h"
 
-#include "OnlineCatchHelper.h"
-
 #define SOCIAL_TAG "[suite_social]"
 #define EG_SOCIAL_SENDFRIENDINVITE_TAG SOCIAL_TAG "[sendfriendinvite]"
 #define EG_SOCIAL_SENDFRIENDINVITEEOS_TAG SOCIAL_TAG "[sendfriendinvite][.EOS]"
-#define EG_SOCIAL_DISABLED_TAG SOCIAL_TAG "[socialdisabled]"
+
 #define SOCIAL_TEST_CASE(x, ...) ONLINE_TEST_CASE(x, SOCIAL_TAG __VA_ARGS__)
 
 SOCIAL_TEST_CASE("Verify that SendFriendInvite returns a fail message if use invalid local user account id", EG_SOCIAL_SENDFRIENDINVITE_TAG)
 {
-	const int32 NumUsersToLogin = 0;
-
 	FSendFriendInvite::Params OpSendFriendInviteParams;
 	FSendFriendInviteHelper::FHelperParams SendFriendInviteHelperParams;
 	SendFriendInviteHelperParams.OpParams = &OpSendFriendInviteParams;
 	SendFriendInviteHelperParams.OpParams->LocalAccountId = FAccountId();
 	SendFriendInviteHelperParams.ExpectedError = TOnlineResult<FSendFriendInvite>(Errors::InvalidParams());
 
-	GetLoginPipeline(NumUsersToLogin)
+	GetPipeline()
 		.EmplaceStep<FSendFriendInviteHelper>(MoveTemp(SendFriendInviteHelperParams));
 
 	RunToCompletion();
@@ -41,7 +37,7 @@ SOCIAL_TEST_CASE("Verify that SendFriendInvite returns a fail message if use inv
 	SendFriendInviteHelperParams.OpParams->TargetAccountId = FAccountId();
 	SendFriendInviteHelperParams.ExpectedError = TOnlineResult<FSendFriendInvite>(Errors::InvalidParams());
 
-	FTestPipeline& LoginPipeline = GetLoginPipeline(AccountId);
+	FTestPipeline& LoginPipeline = GetLoginPipeline({ AccountId });
 
 	SendFriendInviteHelperParams.OpParams->LocalAccountId = AccountId;
 
@@ -55,21 +51,42 @@ SOCIAL_TEST_CASE("Verify that SendFriendInvite returns a fail message if the loc
 {
 	FAccountId FirstAccountId, SecondAccountId;
 
+	int32 UserNumToLogin = 1;
+	TSharedPtr<FPlatformUserId> FirstAccountPlatformUserId = MakeShared<FPlatformUserId>();
+	TSharedPtr<FPlatformUserId> SecondAccountPlatformUserId = MakeShared<FPlatformUserId>();
+	bool bLogout = false;
+
 	FSendFriendInvite::Params OpSendFriendInviteParams;
 	FSendFriendInviteHelper::FHelperParams SendFriendInviteHelperParams;
 	SendFriendInviteHelperParams.OpParams = &OpSendFriendInviteParams;
 	SendFriendInviteHelperParams.ExpectedError = TOnlineResult<FSendFriendInvite>(Errors::NotLoggedIn());
 
-	FTestPipeline& LoginPipeline = GetLoginPipeline(FirstAccountId, SecondAccountId);
+	FTestPipeline& LoginPipeline = GetLoginPipeline(UserNumToLogin, { FirstAccountId, SecondAccountId });
 
 	SendFriendInviteHelperParams.OpParams->LocalAccountId = FirstAccountId;
 	SendFriendInviteHelperParams.OpParams->TargetAccountId = SecondAccountId;
 
-	bool bLogout = false;
-
 	LoginPipeline
-		.EmplaceStep<FAuthLogoutStep>(FPlatformMisc::GetPlatformUserForUserIndex(0))
-		.EmplaceStep<FSendFriendInviteHelper>(MoveTemp(SendFriendInviteHelperParams));
+		.EmplaceLambda([&FirstAccountId, &SecondAccountId, FirstAccountPlatformUserId, SecondAccountPlatformUserId](SubsystemType OnlineSubsystem)
+			{
+				UE::Online::IAuthPtr OnlineAuthPtr = OnlineSubsystem->GetAuthInterface();
+				REQUIRE(OnlineAuthPtr);
+
+				UE::Online::TOnlineResult<UE::Online::FAuthGetLocalOnlineUserByOnlineAccountId> FirstUserPlatfromUserIdResult = OnlineAuthPtr->GetLocalOnlineUserByOnlineAccountId({ FirstAccountId });
+				UE::Online::TOnlineResult<UE::Online::FAuthGetLocalOnlineUserByOnlineAccountId> SecondUserPlatfromUserIdResult = OnlineAuthPtr->GetLocalOnlineUserByOnlineAccountId({ SecondAccountId });
+
+				REQUIRE(FirstUserPlatfromUserIdResult.IsOk());
+				REQUIRE(SecondUserPlatfromUserIdResult.IsOk());
+
+				CHECK(FirstUserPlatfromUserIdResult.TryGetOkValue() != nullptr);
+				CHECK(SecondUserPlatfromUserIdResult.TryGetOkValue() != nullptr);
+
+				*FirstAccountPlatformUserId = FirstUserPlatfromUserIdResult.TryGetOkValue()->AccountInfo->PlatformUserId;
+				*SecondAccountPlatformUserId = SecondUserPlatfromUserIdResult.TryGetOkValue()->AccountInfo->PlatformUserId;
+			})
+		.EmplaceStep<FAuthLogoutStep>(MoveTemp(FirstAccountPlatformUserId))
+		.EmplaceStep<FSendFriendInviteHelper>(MoveTemp(SendFriendInviteHelperParams))
+		.EmplaceStep<FAuthLogoutStep>(MoveTemp(SecondAccountPlatformUserId));
 
 	RunToCompletion(bLogout);
 }
diff --git a/Engine/Source/Programs/Online/OnlineTestsCore/Source/Private/Tests/Stats/StatsTests.cpp b/Engine/Source/Programs/Online/OnlineTestsCore/Source/Private/Tests/Stats/StatsTests.cpp
--- a/Engine/Source/Programs/Online/OnlineTestsCore/Source/Private/Tests/Stats/StatsTests.cpp
+++ b/Engine/Source/Programs/Online/OnlineTestsCore/Source/Private/Tests/Stats/StatsTests.cpp
@@ -1,10 +1,7 @@
 // Copyright Epic Games, Inc. All Rights Reserved.
 
-#include <catch2/catch_test_macros.hpp>
-
 #include "Helpers/Identity/IdentityGetLoginByUserId.h"
 #include "Helpers/Stats/UpdateStatsHelper.h"
-#include "OnlineCatchHelper.h"
 #include "Online/Stats.h"
 #include "Online/StatsCommon.h"
 
@@ -13,7 +10,7 @@
 
 FTestPipeline&& StatsFixture_1User_With_InitialRandomStats(const OnlineTestBase& InOnlineTestBase, FAccountId& AccountId)
 {
-	return InOnlineTestBase.GetLoginPipeline(AccountId)
+	return InOnlineTestBase.GetLoginPipeline({ AccountId })
 		.EmplaceAsyncLambda([&AccountId](FAsyncLambdaResult Promise, SubsystemType Services) {
 			int64 InitValue = (int64)FMath::RandRange(1, 100);
 			TArray<FUserStats> UpdateUsersStats{
@@ -89,7 +86,7 @@ STATS_TEST_CASE("UpdateStats modify different types work as expected", "[tmpExcl
 	TMap<FString, FStatValue> InitialStats;
 	TMap<FString, FStatValue> UpdatedStats;
 
-	GetLoginPipeline(AccountId)
+	GetLoginPipeline({ AccountId })
 		.EmplaceAsyncLambda([&AccountId](FAsyncLambdaResult Promise, SubsystemType Services) {
 			UpdateStats_Fixture(MoveTemp(Promise), Services, AccountId, { { AccountId, { { "Stat_Type_Bool", FStatValue(true)}, {"Stat_Type_Double", FStatValue(9999.999)}}}, });
 		})
@@ -201,7 +198,7 @@ STATS_TEST_CASE("BatchUpdateAndQuery with 2 users", "[MultiAccount]")
 	int64 Expected_Stat_Use_Set_A = (int64)FMath::RandRange(1, 100);
 	int64 Expected_Stat_Use_Set_B = (int64)FMath::RandRange(1, 100);
 
-	GetLoginPipeline(AccountIdA, AccountIdB)
+	GetLoginPipeline({ AccountIdA, AccountIdB })
 		.EmplaceAsyncLambda([&AccountIdA, &Expected_Stat_Use_Set_A](FAsyncLambdaResult Promise, SubsystemType Services) {
 			TArray<FUserStats> UpdateUsersStats{
 				{ AccountIdA, { { "Stat_Use_Set", FStatValue(Expected_Stat_Use_Set_A) } } },
diff --git a/Engine/Source/Programs/Online/OnlineTestsCore/Source/Private/Tests/TestExecHandler.cpp b/Engine/Source/Programs/Online/OnlineTestsCore/Source/Private/Tests/TestExecHandler.cpp
--- a/Engine/Source/Programs/Online/OnlineTestsCore/Source/Private/Tests/TestExecHandler.cpp
+++ b/Engine/Source/Programs/Online/OnlineTestsCore/Source/Private/Tests/TestExecHandler.cpp
@@ -339,7 +339,7 @@ EH_ERROR_TEST_CASE("Test conplex object parsing")
 ONLINE_TEST_CASE("Test new account ID support working", EH_ERROR_TAG)
 {
 	FAccountId Id;
-	GetLoginPipeline(Id);
+	GetLoginPipeline({ Id });
 
 
 	const TCHAR* Cmd = TEXT("0");
diff --git a/Engine/Source/Programs/Online/OnlineTestsCore/Source/Private/Tests/TitleFile/TitleFileTests.cpp b/Engine/Source/Programs/Online/OnlineTestsCore/Source/Private/Tests/TitleFile/TitleFileTests.cpp
--- a/Engine/Source/Programs/Online/OnlineTestsCore/Source/Private/Tests/TitleFile/TitleFileTests.cpp
+++ b/Engine/Source/Programs/Online/OnlineTestsCore/Source/Private/Tests/TitleFile/TitleFileTests.cpp
@@ -1,9 +1,5 @@
 // Copyright Epic Games, Inc. All Rights Reserved.
 
-#include "CoreMinimal.h"
-
-#include <catch2/catch_test_macros.hpp>
-
 #include "Helpers/TitleFile/AsyncTitleFileEnumerateFiles.h"
 #include "Helpers/TitleFile/AsyncTitleFileGetEnumeratedFiles.h"
 #include "Helpers/TitleFile/AsyncTitleFileReadFile.h"
@@ -33,7 +29,7 @@ FTitleFileContentsRef FormatExpectedFileContents(FString FileContents)
 TITLE_FILE_TEST_CASE("EnumerateFiles", TITLE_FILE_ENUMERATE_TAG)
 {
 	FAccountId AccountId;
-	FTestPipeline& Pipe = GetLoginPipeline(AccountId);
+	FTestPipeline& Pipe = GetLoginPipeline({ AccountId });
 
 	FTitleFileEnumerateFiles::Params Params;
 	Params.LocalAccountId = AccountId;
@@ -46,7 +42,7 @@ TITLE_FILE_TEST_CASE("EnumerateFiles", TITLE_FILE_ENUMERATE_TAG)
 TITLE_FILE_TEST_CASE("EnumerateFiles - Invalid User", TITLE_FILE_ENUMERATE_TAG)
 {
 	FAccountId AccountId;
-	FTestPipeline& Pipe = GetLoginPipeline(AccountId);
+	FTestPipeline& Pipe = GetLoginPipeline({ AccountId });
 
 	FTitleFileEnumerateFiles::Params Params;
 
@@ -58,7 +54,7 @@ TITLE_FILE_TEST_CASE("EnumerateFiles - Invalid User", TITLE_FILE_ENUMERATE_TAG)
 TITLE_FILE_TEST_CASE("GetEnumeratedFiles", TITLE_FILE_GET_ENUMERATED_TAG)
 {
 	FAccountId AccountId;
-	FTestPipeline& Pipe = GetLoginPipeline(AccountId);
+	FTestPipeline& Pipe = GetLoginPipeline({ AccountId });
 
 	FTitleFileEnumerateFiles::Params EnumerateParams;
 	EnumerateParams.LocalAccountId = AccountId;
@@ -78,7 +74,7 @@ TITLE_FILE_TEST_CASE("GetEnumeratedFiles", TITLE_FILE_GET_ENUMERATED_TAG)
 TITLE_FILE_TEST_CASE("GetEnumeratedFiles - Invalid user", TITLE_FILE_GET_ENUMERATED_TAG)
 {
 	FAccountId AccountId;
-	FTestPipeline& Pipe = GetLoginPipeline(AccountId);
+	FTestPipeline& Pipe = GetLoginPipeline({ AccountId });
 
 	FTitleFileEnumerateFiles::Params EnumerateParams;
 	EnumerateParams.LocalAccountId = AccountId;
@@ -96,7 +92,7 @@ TITLE_FILE_TEST_CASE("GetEnumeratedFiles - Invalid state", TITLE_FILE_GET_ENUMER
 	DestroyCurrentServiceModule();
 
 	FAccountId AccountId;
-	FTestPipeline& Pipe = GetLoginPipeline(AccountId);
+	FTestPipeline& Pipe = GetLoginPipeline({ AccountId });
 
 	FTitleFileGetEnumeratedFiles::Params GetParams;
 	GetParams.LocalAccountId = AccountId;
@@ -109,7 +105,7 @@ TITLE_FILE_TEST_CASE("GetEnumeratedFiles - Invalid state", TITLE_FILE_GET_ENUMER
 TITLE_FILE_TEST_CASE("ReadFile", TITLE_FILE_READ_TAG)
 {
 	FAccountId AccountId;
-	FTestPipeline& Pipe = GetLoginPipeline(AccountId);
+	FTestPipeline& Pipe = GetLoginPipeline({ AccountId });
 
 	FTitleFileEnumerateFiles::Params EnumerateParams;
 	EnumerateParams.LocalAccountId = AccountId;
@@ -135,7 +131,7 @@ TITLE_FILE_TEST_CASE("ReadFile", TITLE_FILE_READ_TAG)
 TITLE_FILE_TEST_CASE("ReadFile - Invalid User", TITLE_FILE_READ_TAG)
 {
 	FAccountId AccountId;
-	FTestPipeline& Pipe = GetLoginPipeline(AccountId);
+	FTestPipeline& Pipe = GetLoginPipeline({ AccountId });
 
 	FTitleFileEnumerateFiles::Params EnumerateParams;
 	EnumerateParams.LocalAccountId = AccountId;
@@ -153,7 +149,7 @@ TITLE_FILE_TEST_CASE("ReadFile - Empty filename", TITLE_FILE_READ_TAG)
 	DestroyCurrentServiceModule();
 
 	FAccountId AccountId;
-	FTestPipeline& Pipe = GetLoginPipeline(AccountId);
+	FTestPipeline& Pipe = GetLoginPipeline({ AccountId });
 
 	FTitleFileReadFile::Params ReadParams;
 	ReadParams.LocalAccountId = AccountId;
@@ -166,7 +162,7 @@ TITLE_FILE_TEST_CASE("ReadFile - Empty filename", TITLE_FILE_READ_TAG)
 TITLE_FILE_TEST_CASE("ReadFile - Not Found", TITLE_FILE_READ_TAG)
 {
 	FAccountId AccountId;
-	FTestPipeline& Pipe = GetLoginPipeline(AccountId);
+	FTestPipeline& Pipe = GetLoginPipeline({ AccountId });
 
 	FTitleFileEnumerateFiles::Params EnumerateParams;
 	EnumerateParams.LocalAccountId = AccountId;
diff --git a/Engine/Source/Programs/Online/OnlineTestsCore/Source/Private/Tests/UserFile/UserFileTests.cpp b/Engine/Source/Programs/Online/OnlineTestsCore/Source/Private/Tests/UserFile/UserFileTests.cpp
--- a/Engine/Source/Programs/Online/OnlineTestsCore/Source/Private/Tests/UserFile/UserFileTests.cpp
+++ b/Engine/Source/Programs/Online/OnlineTestsCore/Source/Private/Tests/UserFile/UserFileTests.cpp
@@ -1,9 +1,5 @@
 // Copyright Epic Games, Inc. All Rights Reserved.
 
-#include "CoreMinimal.h"
-
-#include <catch2/catch_test_macros.hpp>
-
 #include "Helpers/UserFile/AsyncUserFileEnumerateFiles.h"
 #include "Helpers/UserFile/AsyncUserFileGetEnumeratedFiles.h"
 #include "Helpers/UserFile/AsyncUserFileCopyFile.h"
@@ -53,7 +49,7 @@ void ClearUserFiles(FAccountId AccountId, FTestPipeline& Pipe)
 USER_FILE_TEST_CASE("EnumerateFiles Tests", USER_FILE_ENUMERATE_TAG)
 {
 	FAccountId AccountId;
-	FTestPipeline& Pipe = GetLoginPipeline(AccountId);
+	FTestPipeline& Pipe = GetLoginPipeline({ AccountId });
 
 	ClearUserFiles(AccountId, Pipe);
 
@@ -71,7 +67,7 @@ USER_FILE_TEST_CASE("EnumerateFiles Tests - Invalid user", USER_FILE_ENUMERATE_T
 	FScopeDisableWarningsInLog ScopeLogOnlineSDK{ &LogOnlineServices };
 
 	FAccountId AccountId;
-	FTestPipeline& Pipe = GetLoginPipeline(AccountId);
+	FTestPipeline& Pipe = GetLoginPipeline({ AccountId });
 
 	FUserFileEnumerateFiles::Params EnumerateParams;
 
@@ -84,7 +80,7 @@ USER_FILE_TEST_CASE("EnumerateFiles Tests - Invalid user", USER_FILE_ENUMERATE_T
 USER_FILE_TEST_CASE("GetEnumeratedFiles Tests", USER_FILE_GETENUMERATED_TAG)
 {
 	FAccountId AccountId;
-	FTestPipeline& Pipe = GetLoginPipeline(AccountId);
+	FTestPipeline& Pipe = GetLoginPipeline({ AccountId });
 
 	ClearUserFiles(AccountId, Pipe);
 
@@ -106,7 +102,7 @@ USER_FILE_TEST_CASE("GetEnumeratedFiles Tests - Invalid user", USER_FILE_GETENUM
 	FScopeDisableWarningsInLog ScopeLogOnlineSDK{ &LogOnlineServices };
 
 	FAccountId AccountId;
-	FTestPipeline& Pipe = GetLoginPipeline(AccountId);
+	FTestPipeline& Pipe = GetLoginPipeline({ AccountId });
 
 	FUserFileEnumerateFiles::Params EnumerateParams;
 	EnumerateParams.LocalAccountId = AccountId;
@@ -127,7 +123,7 @@ USER_FILE_TEST_CASE("GetEnumeratedFiles Tests - Invalid State", USER_FILE_GETENU
 	DestroyCurrentServiceModule();
 
 	FAccountId AccountId;
-	FTestPipeline& Pipe = GetLoginPipeline(AccountId);
+	FTestPipeline& Pipe = GetLoginPipeline({ AccountId });
 
 	FUserFileGetEnumeratedFiles::Params GetParams;
 	GetParams.LocalAccountId = AccountId;
@@ -140,7 +136,7 @@ USER_FILE_TEST_CASE("GetEnumeratedFiles Tests - Invalid State", USER_FILE_GETENU
 USER_FILE_TEST_CASE("WriteFile Tests", USER_FILE_WRITE_TAG)
 {
 	FAccountId AccountId;
-	FTestPipeline& Pipe = GetLoginPipeline(AccountId);
+	FTestPipeline& Pipe = GetLoginPipeline({ AccountId });
 
 	ClearUserFiles(AccountId, Pipe);
 
@@ -157,7 +153,7 @@ USER_FILE_TEST_CASE("WriteFile Tests", USER_FILE_WRITE_TAG)
 USER_FILE_TEST_CASE("WriteFile Tests - Empty filename", USER_FILE_WRITE_TAG)
 {
 	FAccountId AccountId;
-	FTestPipeline& Pipe = GetLoginPipeline(AccountId);
+	FTestPipeline& Pipe = GetLoginPipeline({ AccountId });
 
 	FUserFileWriteFile::Params Params;
 	Params.LocalAccountId = AccountId;
@@ -171,7 +167,7 @@ USER_FILE_TEST_CASE("WriteFile Tests - Empty filename", USER_FILE_WRITE_TAG)
 USER_FILE_TEST_CASE("WriteFile Tests - Invalid User", USER_FILE_WRITE_TAG)
 {
 	FAccountId AccountId;
-	FTestPipeline& Pipe = GetLoginPipeline(AccountId);
+	FTestPipeline& Pipe = GetLoginPipeline({ AccountId });
 
 	FUserFileWriteFile::Params Params;
 	Params.Filename = TEXT("TestFileA");
@@ -185,7 +181,7 @@ USER_FILE_TEST_CASE("WriteFile Tests - Invalid User", USER_FILE_WRITE_TAG)
 USER_FILE_TEST_CASE("ReadFile Tests", USER_FILE_READ_TAG)
 {
 	FAccountId AccountId;
-	FTestPipeline& Pipe = GetLoginPipeline(AccountId);
+	FTestPipeline& Pipe = GetLoginPipeline({ AccountId });
 
 	FUserFileWriteFile::Params WriteParams;
 	WriteParams.LocalAccountId = AccountId;
@@ -207,7 +203,7 @@ USER_FILE_TEST_CASE("ReadFile Tests", USER_FILE_READ_TAG)
 USER_FILE_TEST_CASE("ReadFile Tests - Empty filename", USER_FILE_READ_TAG)
 {
 	FAccountId AccountId;
-	FTestPipeline& Pipe = GetLoginPipeline(AccountId);
+	FTestPipeline& Pipe = GetLoginPipeline({ AccountId });
 
 	FUserFileReadFile::Params Params;
 	Params.LocalAccountId = AccountId;
@@ -220,7 +216,7 @@ USER_FILE_TEST_CASE("ReadFile Tests - Empty filename", USER_FILE_READ_TAG)
 USER_FILE_TEST_CASE("ReadFile Tests - Invalid User", USER_FILE_READ_TAG)
 {
 	FAccountId AccountId;
-	FTestPipeline& Pipe = GetLoginPipeline(AccountId);
+	FTestPipeline& Pipe = GetLoginPipeline({ AccountId });
 
 	FUserFileReadFile::Params Params;
 	Params.Filename = TEXT("TestFileA");
@@ -236,7 +232,7 @@ USER_FILE_TEST_CASE("ReadFile Tests - Not Found", USER_FILE_READ_TAG)
 	FScopeDisableWarningsInLog ScopeLogOnlineSDK{ &LogOnlineServices };
 
 	FAccountId AccountId;
-	FTestPipeline& Pipe = GetLoginPipeline(AccountId);
+	FTestPipeline& Pipe = GetLoginPipeline({ AccountId });
 
 	FUserFileReadFile::Params Params;
 	Params.LocalAccountId = AccountId;
@@ -250,7 +246,7 @@ USER_FILE_TEST_CASE("ReadFile Tests - Not Found", USER_FILE_READ_TAG)
 USER_FILE_TEST_CASE("CopyFile Tests", USER_FILE_COPY_TAG)
 {
 	FAccountId AccountId;
-	FTestPipeline& Pipe = GetLoginPipeline(AccountId);
+	FTestPipeline& Pipe = GetLoginPipeline({ AccountId });
 
 	FUserFileWriteFile::Params WriteParams;
 	WriteParams.LocalAccountId = AccountId;
@@ -281,7 +277,7 @@ USER_FILE_TEST_CASE("CopyFile Tests - Invalid User", USER_FILE_COPY_TAG)
 	FScopeDisableWarningsInLog ScopeLogOnlineSDK{ &LogOnlineServices };
 
 	FAccountId AccountId;
-	FTestPipeline& Pipe = GetLoginPipeline(AccountId);
+	FTestPipeline& Pipe = GetLoginPipeline({ AccountId });
 
 	FUserFileCopyFile::Params Params;
 	Params.SourceFilename = TEXT("TestFileA");
@@ -298,7 +294,7 @@ USER_FILE_TEST_CASE("CopyFile Tests - Empty target", USER_FILE_COPY_TAG)
 	FScopeDisableWarningsInLog ScopeLogOnlineSDK{ &LogOnlineServices };
 
 	FAccountId AccountId;
-	FTestPipeline& Pipe = GetLoginPipeline(AccountId);
+	FTestPipeline& Pipe = GetLoginPipeline({ AccountId });
 
 	FUserFileCopyFile::Params Params;
 	Params.LocalAccountId = AccountId;
@@ -315,7 +311,7 @@ USER_FILE_TEST_CASE("CopyFile Tests - Empty source", USER_FILE_COPY_TAG)
 	FScopeDisableWarningsInLog ScopeLogOnlineSDK{ &LogOnlineServices };
 
 	FAccountId AccountId;
-	FTestPipeline& Pipe = GetLoginPipeline(AccountId);
+	FTestPipeline& Pipe = GetLoginPipeline({ AccountId });
 
 	FUserFileCopyFile::Params Params;
 	Params.LocalAccountId = AccountId;
@@ -332,7 +328,7 @@ USER_FILE_TEST_CASE("CopyFile Tests - Not Found", USER_FILE_COPY_TAG)
 	FScopeDisableWarningsInLog ScopeLogOnlineSDK{ &LogOnlineServices };
 
 	FAccountId AccountId;
-	FTestPipeline& Pipe = GetLoginPipeline(AccountId);
+	FTestPipeline& Pipe = GetLoginPipeline({ AccountId });
 
 	FUserFileCopyFile::Params Params;
 	Params.LocalAccountId = AccountId;
@@ -347,7 +343,7 @@ USER_FILE_TEST_CASE("CopyFile Tests - Not Found", USER_FILE_COPY_TAG)
 USER_FILE_TEST_CASE("DeleteFile Tests", USER_FILE_DELETE_TAG)
 {
 	FAccountId AccountId;
-	FTestPipeline& Pipe = GetLoginPipeline(AccountId);
+	FTestPipeline& Pipe = GetLoginPipeline({ AccountId });
 
 	FUserFileDeleteFile::Params Params;
 	Params.LocalAccountId = AccountId;
@@ -361,7 +357,7 @@ USER_FILE_TEST_CASE("DeleteFile Tests", USER_FILE_DELETE_TAG)
 USER_FILE_TEST_CASE("DeleteFile Tests - Invalid User", USER_FILE_DELETE_TAG)
 {
 	FAccountId AccountId;
-	FTestPipeline& Pipe = GetLoginPipeline(AccountId);
+	FTestPipeline& Pipe = GetLoginPipeline({ AccountId });
 
 	FUserFileDeleteFile::Params Params;
 	Params.Filename = TEXT("TestFileA");
@@ -374,7 +370,7 @@ USER_FILE_TEST_CASE("DeleteFile Tests - Invalid User", USER_FILE_DELETE_TAG)
 USER_FILE_TEST_CASE("DeleteFile Tests - Empty filename", USER_FILE_DELETE_TAG)
 {
 	FAccountId AccountId;
-	FTestPipeline& Pipe = GetLoginPipeline(AccountId);
+	FTestPipeline& Pipe = GetLoginPipeline({ AccountId });
 
 	FUserFileDeleteFile::Params Params;
 	Params.LocalAccountId = AccountId;
diff --git a/Engine/Source/Programs/Online/OnlineTestsCore/Source/Public/Helpers/Auth/AuthLogin.h b/Engine/Source/Programs/Online/OnlineTestsCore/Source/Public/Helpers/Auth/AuthLogin.h
--- a/Engine/Source/Programs/Online/OnlineTestsCore/Source/Public/Helpers/Auth/AuthLogin.h
+++ b/Engine/Source/Programs/Online/OnlineTestsCore/Source/Public/Helpers/Auth/AuthLogin.h
@@ -1,8 +1,10 @@
 // Copyright Epic Games, Inc. All Rights Reserved.
 #pragma once
+
 #include "CoreMinimal.h"
 #include "TestDriver.h"
 #include "TestHarness.h"
+#include "AsyncTestStep.h"
 
 #include "Online/AuthCommon.h"
 #include "Online/OnlineAsyncOp.h"
@@ -10,51 +12,33 @@
 #include "Online/OnlineServicesCommon.h"
 #include "OnlineCatchHelper.h"
 
-using namespace UE::Online;
-
-struct FAuthLoginStep : public FTestPipeline::FStep
+struct FAuthLoginStep : public FAsyncTestStep
 {
 	FAuthLoginStep(UE::Online::FAuthLogin::Params&& InLocalAccount)
 		: LocalAccount(MoveTemp(InLocalAccount))
 	{
 	}
 
-	virtual ~FAuthLoginStep()
-	{
-		if (OnlineAuthPtr != nullptr)
-		{
-			OnlineAuthPtr = nullptr;
-		}
-	};
-
-	enum class EState { Init, LoginCalled, Done } State = EState::Init;
-
+	virtual ~FAuthLoginStep() = default;
 
-	virtual EContinuance Tick(SubsystemType OnlineSubsystem) override
+	virtual void Run(FAsyncStepResult Promise, SubsystemType Services) override
 	{
-		OnlineAuthPtr = OnlineSubsystem->GetAuthInterface();
+		OnlineAuthPtr = Services->GetAuthInterface();
+		REQUIRE(OnlineAuthPtr);
+
+		FPlatformUserId PlatformUserId = LocalAccount.PlatformUserId;
 
-		switch (State)
+		TOnlineResult<FAuthGetLocalOnlineUserByPlatformUserId> LocalOnlineUserResult = OnlineAuthPtr->GetLocalOnlineUserByPlatformUserId({PlatformUserId});
+		if (LocalOnlineUserResult.IsOk() && LocalOnlineUserResult.GetOkValue().AccountInfo->LoginStatus == UE::Online::ELoginStatus::LoggedIn)
 		{
-		case EState::Init:
+			Promise->SetValue(true);
+		}
+		else
 		{
-			State = EState::LoginCalled;
-
-			FPlatformUserId PlatformUserId = LocalAccount.PlatformUserId;
-
-			TOnlineResult<FAuthGetLocalOnlineUserByPlatformUserId> LocalOnlineUserResult = OnlineAuthPtr->GetLocalOnlineUserByPlatformUserId({PlatformUserId});
-			if (LocalOnlineUserResult.IsOk() && LocalOnlineUserResult.GetOkValue().AccountInfo->LoginStatus == UE::Online::ELoginStatus::LoggedIn)
-			{
-				State = EState::Done;
-			}
-			else
-			{
-				OnlineAuthPtr->Login(MoveTemp(LocalAccount))
-				 .OnComplete([this, PlatformUserId](const TOnlineResult<FAuthLogin> Op) mutable
+			OnlineAuthPtr->Login(MoveTemp(LocalAccount))
+				.OnComplete([this, PlatformUserId, Promise = MoveTemp(Promise)](const TOnlineResult<FAuthLogin> Op) mutable
 				{
-					CHECK(State == EState::LoginCalled);
 					CHECK_OP_EQ(Op, Errors::NotImplemented());
-
 					if(Op.IsOk())
 					{
 						CHECK_OP(OnlineAuthPtr->GetLocalOnlineUserByOnlineAccountId({Op.GetOkValue().AccountInfo->AccountId}));
@@ -66,28 +50,13 @@ struct FAuthLoginStep : public FTestPipeline::FStep
 						CHECK_OP(OnlineUserResult);
 					}
 
-					State = EState::Done;
+					Promise->SetValue(true);
 				});
-			}
-
-			break;
-		}
-		case EState::LoginCalled:
-		{
-			break;
 		}
-		case EState::Done:
-		{
-			return EContinuance::Done;
-		}
-		}
-
-		return EContinuance::ContinueStepping;
 	}
 
 protected:
 	int32 LocalUserNum = 0;
 	UE::Online::FAuthLogin::Params LocalAccount;
 	UE::Online::IAuthPtr OnlineAuthPtr = nullptr;
-	FDelegateHandle	OnLoginCompleteDelegateHandle;
 };
\ No newline at end of file
diff --git a/Engine/Source/Programs/Online/OnlineTestsCore/Source/Public/Helpers/Auth/AuthLogout.h b/Engine/Source/Programs/Online/OnlineTestsCore/Source/Public/Helpers/Auth/AuthLogout.h
--- a/Engine/Source/Programs/Online/OnlineTestsCore/Source/Public/Helpers/Auth/AuthLogout.h
+++ b/Engine/Source/Programs/Online/OnlineTestsCore/Source/Public/Helpers/Auth/AuthLogout.h
@@ -4,72 +4,56 @@
 #include "CoreMinimal.h"
 #include "TestDriver.h"
 #include "TestHarness.h"
+#include "AsyncTestStep.h"
 
 #include "Online/AuthCommon.h"
 #include "Online/OnlineAsyncOp.h"
 #include "Online/OnlineErrorDefinitions.h"
 #include "Online/OnlineServicesCommon.h"
 
-struct FAuthLogoutStep : public FTestPipeline::FStep
+struct FAuthLogoutStep : public FAsyncTestStep
 {
 	FAuthLogoutStep(FPlatformUserId InPlatformUserId)
 		: PlatformUserId(InPlatformUserId)
 	{}
 
-	virtual ~FAuthLogoutStep()
+	FAuthLogoutStep(TSharedPtr<FPlatformUserId>&& InPlatformUserIdPtr)
+		: PlatformUserIdPtr(MoveTemp(InPlatformUserIdPtr))
 	{
-		OnlineAuthPtr = nullptr;
-	};
-
-
-	enum class EState { Init, LogoutCalled, Done } State = EState::Init;
+	}
 
+	virtual ~FAuthLogoutStep() = default;
 
-	virtual EContinuance Tick(SubsystemType OnlineSubsystem) override
+	virtual void Run(FAsyncStepResult Promise, SubsystemType Services) override
 	{
-		OnlineAuthPtr = OnlineSubsystem->GetAuthInterface();
+		OnlineAuthPtr = Services->GetAuthInterface();
+		REQUIRE(OnlineAuthPtr);
 
-		switch (State)
+		if (PlatformUserIdPtr)
 		{
-		case EState::Init:
-		{
-			State = EState::LogoutCalled;
-			TOnlineResult<FAuthGetLocalOnlineUserByPlatformUserId> AccountId = OnlineAuthPtr->GetLocalOnlineUserByPlatformUserId({PlatformUserId});
-			CAPTURE(ToLogString(AccountId), PlatformUserId);
-			CHECK_OP(AccountId);
-			if(AccountId.IsOk())
-			{
-				OnlineAuthPtr->Logout({ AccountId.GetOkValue().AccountInfo->AccountId})
-					.OnComplete([this](const TOnlineResult<FAuthLogout> Op) mutable
-					{
-						CHECK(State == EState::LogoutCalled);
-						// Some implementations do not implement an explicit login / logout.
-						CHECK_OP_EQ(Op, Errors::NotImplemented());
-
-						State = EState::Done;
-					});
-			}
-			else
-			{
-				State = EState::Done;
-			}
-
-			break;
+			PlatformUserId = *PlatformUserIdPtr;
 		}
-		case EState::LogoutCalled:
+		
+		TOnlineResult<FAuthGetLocalOnlineUserByPlatformUserId> AccountId = OnlineAuthPtr->GetLocalOnlineUserByPlatformUserId({ PlatformUserId });
+
+		CAPTURE(ToLogString(AccountId), PlatformUserId);
+		CHECK_OP(AccountId);
+		if(AccountId.IsOk())
 		{
-			break;
+			OnlineAuthPtr->Logout({ AccountId.GetOkValue().AccountInfo->AccountId})
+				.OnComplete([this, Promise = MoveTemp(Promise)](const TOnlineResult<FAuthLogout> Op) mutable
+				{
+					CHECK_OP_EQ(Op, Errors::NotImplemented());
+					Promise->SetValue(true);
+				});
 		}
-		case EState::Done:
+		else
 		{
-			return EContinuance::Done;
+			Promise->SetValue(true);
 		}
-		}
-
-		return EContinuance::ContinueStepping;
 	}
 protected:
+	TSharedPtr<FPlatformUserId> PlatformUserIdPtr = nullptr;
 	FPlatformUserId PlatformUserId;
 	UE::Online::IAuthPtr OnlineAuthPtr = nullptr;
-	FDelegateHandle	OnLogoutCompleteDelegateHandle;
 };
\ No newline at end of file
diff --git a/Engine/Source/Programs/Online/OnlineTestsCore/Source/Public/Helpers/Commerce/CheckoutHelper.h b/Engine/Source/Programs/Online/OnlineTestsCore/Source/Public/Helpers/Commerce/CheckoutHelper.h
new file mode 100644
--- /dev/null
+++ b/Engine/Source/Programs/Online/OnlineTestsCore/Source/Public/Helpers/Commerce/CheckoutHelper.h
@@ -0,0 +1,53 @@
+// Copyright Epic Games, Inc. All Rights Reserved.
+
+#pragma once
+
+#include "Online/CommerceCommon.h"
+#include "AsyncTestStep.h"
+#include "OnlineCatchHelper.h"
+
+struct FCheckoutHelper : public FAsyncTestStep
+{
+	using ParamsType = UE::Online::FCommerceCheckout::Params;
+	using ResultType = UE::Online::TOnlineResult<FCommerceCheckout>;
+
+	struct FHelperParams
+	{
+		ParamsType* OpParams = nullptr;
+		TOptional<ResultType> ExpectedError;
+	};
+
+	FCheckoutHelper(FHelperParams&& InHelperParams)
+		: HelperParams(MoveTemp(InHelperParams))
+	{
+		REQUIRE(HelperParams.OpParams);
+		REQUIRE((!HelperParams.ExpectedError.IsSet() || HelperParams.ExpectedError->IsError()));
+	}
+
+	virtual ~FCheckoutHelper() = default;
+
+	virtual void Run(FAsyncStepResult Promise, SubsystemType Services) override
+	{
+		CommerceInterface = Services->GetCommerceInterface();
+		REQUIRE(CommerceInterface);
+
+		UE::Online::TOnlineAsyncOpHandle<UE::Online::FCommerceCheckout> CheckoutResult = CommerceInterface->Checkout(MoveTemp(*HelperParams.OpParams))
+			.OnComplete([this, Promise = MoveTemp(Promise)](const ResultType& Result)
+				{
+					if (HelperParams.ExpectedError.IsSet())
+					{
+						REQUIRE_OP_EQ(Result, HelperParams.ExpectedError->GetErrorValue());
+					}
+					else
+					{
+						CHECK(Result.GetOkValue().TransactionId.IsSet());
+					}
+					Promise->SetValue(true);
+				});
+	}
+
+protected:
+	FHelperParams HelperParams;
+	UE::Online::ICommercePtr CommerceInterface = nullptr;
+	TOptional<FString> TransactionId;
+};
\ No newline at end of file
diff --git a/Engine/Source/Programs/Online/OnlineTestsCore/Source/Public/Helpers/Commerce/CommerceCheckoutHelper.h b/Engine/Source/Programs/Online/OnlineTestsCore/Source/Public/Helpers/Commerce/CommerceCheckoutHelper.h
deleted file mode 100644
diff --git a/Engine/Source/Programs/Online/OnlineTestsCore/Source/Public/Helpers/Commerce/CommerceGetEntitlementsHelper.h b/Engine/Source/Programs/Online/OnlineTestsCore/Source/Public/Helpers/Commerce/CommerceGetEntitlementsHelper.h
deleted file mode 100644
diff --git a/Engine/Source/Programs/Online/OnlineTestsCore/Source/Public/Helpers/Commerce/CommerceGetOffersByIdHelper.h b/Engine/Source/Programs/Online/OnlineTestsCore/Source/Public/Helpers/Commerce/CommerceGetOffersByIdHelper.h
deleted file mode 100644
diff --git a/Engine/Source/Programs/Online/OnlineTestsCore/Source/Public/Helpers/Commerce/CommerceGetOffersHelper.h b/Engine/Source/Programs/Online/OnlineTestsCore/Source/Public/Helpers/Commerce/CommerceGetOffersHelper.h
deleted file mode 100644
diff --git a/Engine/Source/Programs/Online/OnlineTestsCore/Source/Public/Helpers/Commerce/CommerceQueryEntitlementsHelper.h b/Engine/Source/Programs/Online/OnlineTestsCore/Source/Public/Helpers/Commerce/CommerceQueryEntitlementsHelper.h
deleted file mode 100644
diff --git a/Engine/Source/Programs/Online/OnlineTestsCore/Source/Public/Helpers/Commerce/CommerceQueryOffersByIdHelper.h b/Engine/Source/Programs/Online/OnlineTestsCore/Source/Public/Helpers/Commerce/CommerceQueryOffersByIdHelper.h
deleted file mode 100644
diff --git a/Engine/Source/Programs/Online/OnlineTestsCore/Source/Public/Helpers/Commerce/CommerceQueryOffersHelper.h b/Engine/Source/Programs/Online/OnlineTestsCore/Source/Public/Helpers/Commerce/CommerceQueryOffersHelper.h
deleted file mode 100644
diff --git a/Engine/Source/Programs/Online/OnlineTestsCore/Source/Public/Helpers/Commerce/GetEntitlementsHelper.h b/Engine/Source/Programs/Online/OnlineTestsCore/Source/Public/Helpers/Commerce/GetEntitlementsHelper.h
new file mode 100644
--- /dev/null
+++ b/Engine/Source/Programs/Online/OnlineTestsCore/Source/Public/Helpers/Commerce/GetEntitlementsHelper.h
@@ -0,0 +1,57 @@
+// Copyright Epic Games, Inc. All Rights Reserved.
+
+#pragma once
+
+#include "Online/CommerceCommon.h"
+#include "AsyncTestStep.h"
+#include "OnlineCatchHelper.h"
+
+struct FGetEntitlementsHelper : public FAsyncTestStep
+{
+	using ParamsType = UE::Online::FCommerceGetEntitlements::Params;
+	using ResultType = UE::Online::TOnlineResult<UE::Online::FCommerceGetEntitlements>;
+
+	struct FHelperParams
+	{
+		ParamsType* OpParams = nullptr;
+		TOptional<ResultType> ExpectedError;
+	};
+
+	FGetEntitlementsHelper(FHelperParams&& InHelperParams, const TOptional<uint32_t> InExpectedEntitlementsNum = TOptional<uint32_t>())
+		: HelperParams(MoveTemp(InHelperParams))
+		, ExpectedEntitlementsNum(InExpectedEntitlementsNum)
+	{
+		REQUIRE(HelperParams.OpParams);
+		REQUIRE((!HelperParams.ExpectedError.IsSet() || HelperParams.ExpectedError->IsError()));
+	}
+
+	virtual ~FGetEntitlementsHelper() = default;
+
+	virtual void Run(FAsyncStepResult Promise, SubsystemType Services) override
+	{
+		CommerceInterface = Services->GetCommerceInterface();
+		REQUIRE(CommerceInterface);
+
+		ResultType Result = CommerceInterface->GetEntitlements(MoveTemp(*HelperParams.OpParams));
+
+		if (HelperParams.ExpectedError.IsSet())
+		{
+			REQUIRE_OP_EQ(Result, HelperParams.ExpectedError->GetErrorValue());
+		}
+		else if (ExpectedEntitlementsNum.IsSet())
+		{
+			CHECK(Result.GetOkValue().Entitlements.Num() == ExpectedEntitlementsNum.GetValue());
+		}
+		else
+		{
+			CHECK(Result.GetOkValue().Entitlements.IsEmpty());
+		}
+		Promise->SetValue(true);
+	}
+
+protected:
+	FHelperParams HelperParams;
+	TOptional<uint32_t> ExpectedEntitlementsNum = TOptional<uint32_t>();
+	UE::Online::ICommercePtr CommerceInterface = nullptr;
+	TArray<FEntitlement> Entitlements;
+};
\ No newline at end of file
diff --git a/Engine/Source/Programs/Online/OnlineTestsCore/Source/Public/Helpers/Commerce/GetOffersByIdHelper.h b/Engine/Source/Programs/Online/OnlineTestsCore/Source/Public/Helpers/Commerce/GetOffersByIdHelper.h
new file mode 100644
--- /dev/null
+++ b/Engine/Source/Programs/Online/OnlineTestsCore/Source/Public/Helpers/Commerce/GetOffersByIdHelper.h
@@ -0,0 +1,57 @@
+// Copyright Epic Games, Inc. All Rights Reserved.
+
+#pragma once
+
+#include "Online/CommerceCommon.h"
+#include "AsyncTestStep.h"
+#include "OnlineCatchHelper.h"
+
+struct FGetOffersByIdHelper : public FAsyncTestStep
+{
+	using ParamsType = UE::Online::FCommerceGetOffersById::Params;
+	using ResultType = UE::Online::TOnlineResult<UE::Online::FCommerceGetOffersById>;
+
+	struct FHelperParams
+	{
+		ParamsType* OpParams = nullptr;
+		TOptional<ResultType> ExpectedError;
+	};
+
+	FGetOffersByIdHelper(FHelperParams&& InHelperParams, const TOptional<uint32_t> InExpectedOffersNum = TOptional<uint32_t>())
+		: HelperParams(MoveTemp(InHelperParams))
+		, ExpectedOffersNum(InExpectedOffersNum)
+	{
+		REQUIRE(HelperParams.OpParams);
+		REQUIRE((!HelperParams.ExpectedError.IsSet() || HelperParams.ExpectedError->IsError()));
+	}
+
+	virtual ~FGetOffersByIdHelper() = default;
+
+	virtual void Run(FAsyncStepResult Promise, SubsystemType Services) override
+	{
+		CommerceInterface = Services->GetCommerceInterface();
+		REQUIRE(CommerceInterface);
+
+		ResultType Result = CommerceInterface->GetOffersById(MoveTemp(*HelperParams.OpParams));
+
+		if (HelperParams.ExpectedError.IsSet())
+		{
+			REQUIRE_OP_EQ(Result, HelperParams.ExpectedError->GetErrorValue());
+		}
+		else if (ExpectedOffersNum.IsSet())
+		{
+			CHECK(Result.GetOkValue().Offers.Num() == ExpectedOffersNum.GetValue());
+		}
+		else
+		{
+			CHECK(Result.GetOkValue().Offers.IsEmpty());
+		}
+		Promise->SetValue(true);
+	}
+
+protected:
+	FHelperParams HelperParams;
+	TOptional<uint32_t> ExpectedOffersNum = TOptional<uint32_t>();
+	UE::Online::ICommercePtr CommerceInterface = nullptr;
+	TArray<FOffer> Offers;
+};
\ No newline at end of file
diff --git a/Engine/Source/Programs/Online/OnlineTestsCore/Source/Public/Helpers/Commerce/GetOffersHelper.h b/Engine/Source/Programs/Online/OnlineTestsCore/Source/Public/Helpers/Commerce/GetOffersHelper.h
new file mode 100644
--- /dev/null
+++ b/Engine/Source/Programs/Online/OnlineTestsCore/Source/Public/Helpers/Commerce/GetOffersHelper.h
@@ -0,0 +1,57 @@
+// Copyright Epic Games, Inc. All Rights Reserved.
+
+#pragma once
+
+#include "Online/CommerceCommon.h"
+#include "AsyncTestStep.h"
+#include "OnlineCatchHelper.h"
+
+struct FGetOffersHelper : public FAsyncTestStep
+{
+	using ParamsType = UE::Online::FCommerceGetOffers::Params;
+	using ResultType = UE::Online::TOnlineResult<UE::Online::FCommerceGetOffers>;
+
+	struct FHelperParams
+	{
+		ParamsType* OpParams = nullptr;
+		TOptional<ResultType> ExpectedError;
+	};
+
+	FGetOffersHelper(FHelperParams&& InHelperParams, const TOptional<uint32_t> InExpectedOffersNum = TOptional<uint32_t>())
+		: HelperParams(MoveTemp(InHelperParams))
+		, ExpectedOffersNum(InExpectedOffersNum)
+	{
+		REQUIRE(HelperParams.OpParams);
+		REQUIRE((!HelperParams.ExpectedError.IsSet() || HelperParams.ExpectedError->IsError()));
+	}
+
+	virtual ~FGetOffersHelper() = default;
+
+	virtual void Run(FAsyncStepResult Promise, SubsystemType Services) override
+	{
+		CommerceInterface = Services->GetCommerceInterface();
+		REQUIRE(CommerceInterface);
+
+		ResultType Result = CommerceInterface->GetOffers(MoveTemp(*HelperParams.OpParams));
+
+		if (HelperParams.ExpectedError.IsSet())
+		{
+			REQUIRE_OP_EQ(Result, HelperParams.ExpectedError->GetErrorValue());
+		}
+		else if (ExpectedOffersNum.IsSet())
+		{
+			CHECK(Result.GetOkValue().Offers.Num() == ExpectedOffersNum.GetValue());
+		}
+		else
+		{
+			CHECK(Result.GetOkValue().Offers.IsEmpty());
+		}
+		Promise->SetValue(true);
+	}
+
+protected:
+	FHelperParams HelperParams;
+	TOptional<uint32_t> ExpectedOffersNum = TOptional<uint32_t>();
+	UE::Online::ICommercePtr CommerceInterface = nullptr;
+	TArray<FOffer> Offers;
+};
\ No newline at end of file
diff --git a/Engine/Source/Programs/Online/OnlineTestsCore/Source/Public/Helpers/Commerce/QueryEntitlementsHelper.h b/Engine/Source/Programs/Online/OnlineTestsCore/Source/Public/Helpers/Commerce/QueryEntitlementsHelper.h
new file mode 100644
--- /dev/null
+++ b/Engine/Source/Programs/Online/OnlineTestsCore/Source/Public/Helpers/Commerce/QueryEntitlementsHelper.h
@@ -0,0 +1,52 @@
+// Copyright Epic Games, Inc. All Rights Reserved.
+
+#pragma once
+
+#include "Online/CommerceCommon.h"
+#include "AsyncTestStep.h"
+#include "OnlineCatchHelper.h"
+
+struct FQueryEntitlementsHelper : public FAsyncTestStep
+{
+	using ParamsType = UE::Online::FCommerceQueryEntitlements::Params;
+	using ResultType = UE::Online::TOnlineResult<UE::Online::FCommerceQueryEntitlements>;
+
+	struct FHelperParams
+	{
+		ParamsType* OpParams = nullptr;
+		TOptional<ResultType> ExpectedError;
+	};
+
+	FQueryEntitlementsHelper(FHelperParams&& InHelperParams)
+		: HelperParams(MoveTemp(InHelperParams))
+	{
+		REQUIRE(HelperParams.OpParams);
+		REQUIRE((!HelperParams.ExpectedError.IsSet() || HelperParams.ExpectedError->IsError()));
+	}
+
+	virtual ~FQueryEntitlementsHelper() = default;
+
+	virtual void Run(FAsyncStepResult Promise, SubsystemType Services) override
+	{
+		CommerceInterface = Services->GetCommerceInterface();
+		REQUIRE(CommerceInterface);
+
+		CommerceInterface->QueryEntitlements(MoveTemp(*HelperParams.OpParams))
+			.OnComplete([this, Promise = MoveTemp(Promise)](const ResultType& Result)
+				{
+					if (!HelperParams.ExpectedError.IsSet())
+					{
+						REQUIRE(Result.IsOk());
+					}
+					else
+					{
+						REQUIRE_OP_EQ(Result, HelperParams.ExpectedError->GetErrorValue());
+					}
+					Promise->SetValue(true);
+				});
+	}
+
+protected:
+	FHelperParams HelperParams;
+	UE::Online::ICommercePtr CommerceInterface = nullptr;
+};
\ No newline at end of file
diff --git a/Engine/Source/Programs/Online/OnlineTestsCore/Source/Public/Helpers/Commerce/QueryOffersByIdHelper.h b/Engine/Source/Programs/Online/OnlineTestsCore/Source/Public/Helpers/Commerce/QueryOffersByIdHelper.h
new file mode 100644
--- /dev/null
+++ b/Engine/Source/Programs/Online/OnlineTestsCore/Source/Public/Helpers/Commerce/QueryOffersByIdHelper.h
@@ -0,0 +1,52 @@
+// Copyright Epic Games, Inc. All Rights Reserved.
+
+#pragma once
+
+#include "Online/CommerceCommon.h"
+#include "AsyncTestStep.h"
+#include "OnlineCatchHelper.h"
+
+struct FQueryOffersByIdHelper : public FAsyncTestStep
+{
+	using ParamsType = UE::Online::FCommerceQueryOffersById::Params;
+	using ResultType = UE::Online::TOnlineResult<UE::Online::FCommerceQueryOffersById>;
+
+	struct FHelperParams
+	{
+		ParamsType* OpParams = nullptr;
+		TOptional<ResultType> ExpectedError;
+	};
+
+	FQueryOffersByIdHelper(FHelperParams&& InHelperParams)
+		: HelperParams(MoveTemp(InHelperParams))
+	{
+		REQUIRE(HelperParams.OpParams);
+		REQUIRE((!HelperParams.ExpectedError.IsSet() || HelperParams.ExpectedError->IsError()));
+	}
+
+	virtual ~FQueryOffersByIdHelper() = default;
+
+	virtual void Run(FAsyncStepResult Promise, SubsystemType Services) override
+	{
+		CommerceInterface = Services->GetCommerceInterface();
+		REQUIRE(CommerceInterface);
+
+		CommerceInterface->QueryOffersById(MoveTemp(*HelperParams.OpParams))
+			.OnComplete([this, Promise = MoveTemp(Promise)](const ResultType& Result)
+				{
+					if (!HelperParams.ExpectedError.IsSet())
+					{
+						REQUIRE(Result.IsOk());
+					}
+					else
+					{
+						REQUIRE_OP_EQ(Result, HelperParams.ExpectedError->GetErrorValue());
+					}
+					Promise->SetValue(true);
+				});
+	}
+
+protected:
+	FHelperParams HelperParams;
+	UE::Online::ICommercePtr CommerceInterface = nullptr;
+};
\ No newline at end of file
diff --git a/Engine/Source/Programs/Online/OnlineTestsCore/Source/Public/Helpers/Commerce/QueryOffersHelper.h b/Engine/Source/Programs/Online/OnlineTestsCore/Source/Public/Helpers/Commerce/QueryOffersHelper.h
new file mode 100644
--- /dev/null
+++ b/Engine/Source/Programs/Online/OnlineTestsCore/Source/Public/Helpers/Commerce/QueryOffersHelper.h
@@ -0,0 +1,52 @@
+// Copyright Epic Games, Inc. All Rights Reserved.
+
+#pragma once
+
+#include "Online/CommerceCommon.h"
+#include "AsyncTestStep.h"
+#include "OnlineCatchHelper.h"
+
+struct FQueryOffersHelper : public FAsyncTestStep
+{
+	using ParamsType = UE::Online::FCommerceQueryOffers::Params;
+	using ResultType = UE::Online::TOnlineResult<UE::Online::FCommerceQueryOffers>;
+
+	struct FHelperParams
+	{
+		ParamsType* OpParams = nullptr;
+		TOptional<ResultType> ExpectedError;
+	};
+
+	FQueryOffersHelper(FHelperParams&& InHelperParams)
+		: HelperParams(MoveTemp(InHelperParams))
+	{
+		REQUIRE(HelperParams.OpParams);
+		REQUIRE((!HelperParams.ExpectedError.IsSet() || HelperParams.ExpectedError->IsError()));
+	}
+
+	virtual ~FQueryOffersHelper() = default;
+
+	virtual void Run(FAsyncStepResult Promise, SubsystemType Services) override
+	{
+		CommerceInterface = Services->GetCommerceInterface();
+		REQUIRE(CommerceInterface);
+
+		CommerceInterface->QueryOffers(MoveTemp(*HelperParams.OpParams))
+			.OnComplete([this, Promise = MoveTemp(Promise)](const ResultType& Result)
+				{
+					if (!HelperParams.ExpectedError.IsSet())
+					{
+						REQUIRE(Result.IsOk());
+					}
+					else
+					{
+						REQUIRE_OP_EQ(Result, HelperParams.ExpectedError->GetErrorValue());
+					}
+					Promise->SetValue(true);
+				});
+	}
+
+protected:
+	FHelperParams HelperParams;
+	UE::Online::ICommercePtr CommerceInterface = nullptr;
+};
\ No newline at end of file
diff --git a/Engine/Source/Programs/Online/OnlineTestsCore/Source/Public/Helpers/Presence/BatchQueryPresenceHelper.h b/Engine/Source/Programs/Online/OnlineTestsCore/Source/Public/Helpers/Presence/BatchQueryPresenceHelper.h
--- a/Engine/Source/Programs/Online/OnlineTestsCore/Source/Public/Helpers/Presence/BatchQueryPresenceHelper.h
+++ b/Engine/Source/Programs/Online/OnlineTestsCore/Source/Public/Helpers/Presence/BatchQueryPresenceHelper.h
@@ -3,6 +3,7 @@
 
 #include "AsyncTestStep.h"
 #include "Online/Presence.h"
+#include "OnlineCatchHelper.h"
 
 using namespace UE::Online;
 class FBatchQueryPresenceHelper : public FAsyncTestStep
diff --git a/Engine/Source/Programs/Online/OnlineTestsCore/Source/Public/OnlineCatchHelper.h b/Engine/Source/Programs/Online/OnlineTestsCore/Source/Public/OnlineCatchHelper.h
--- a/Engine/Source/Programs/Online/OnlineTestsCore/Source/Public/OnlineCatchHelper.h
+++ b/Engine/Source/Programs/Online/OnlineTestsCore/Source/Public/OnlineCatchHelper.h
@@ -20,6 +20,7 @@
 #include <catch2/interfaces/catch_interfaces_registry_hub.hpp>
 #include <catch2/catch_active_test.hpp>
 #include <string>
+
 using namespace UE::Online;
 
 class OnlineTestBase : public Catch::ITestInvoker
@@ -36,24 +37,17 @@ public:
 	/* Unloads all necessary services for the current test run */
 	static void UnloadServiceModules();
 
-	FTestPipeline& GetLoginPipeline(uint32 NumUsersToLogin = 1) const;
-	FTestPipeline& GetLoginPipeline(uint32 UserNumToLogin, FAccountId& OutAccountId) const;
+#if ONLINETESTS_USEEXTERNAUTH
+	/* Returning accounts to the pool */
+	static bool ReturnAccounts(int32 TemplateRequestedNum = -1);
+#endif
+
+	FTestPipeline& GetLoginPipeline(std::initializer_list<std::reference_wrapper<FAccountId>> AccountIds) const;
+
+	FTestPipeline& GetLoginPipeline(uint32 UserNumToLogin, std::initializer_list<std::reference_wrapper<FAccountId>> AccountIds) const;
 
 	typedef UE::Online::FAccountId FAccountId;
 public:
-	/* Creates a login pipeline and logs in user as many users as FAccountIds passed in, immediately assigns FAccountIds to the passed in arguments.
-		This function relies on FAuthLoginStep to be executed by RunToCompletion(...) so will flush any queued steps in the test. */
-	FTestPipeline& GetLoginPipelineArray(std::initializer_list<std::reference_wrapper<FAccountId>> AccountIdsArr) const;
-	/* Creates a login pipeline and logs in user, immediately assigns FAccountIds to the passed in arguments for one account */
-	FTestPipeline& GetLoginPipeline(FAccountId& OutAccountId) const;
-	/* Creates a login pipeline and logs in user, immediately assigns FAccountIds to the passed in arguments for two accounts */
-	FTestPipeline& GetLoginPipeline(FAccountId& OutAccountId, FAccountId& OutAccountId2 ) const;
-	/* Creates a login pipeline and logs in user, immediately assigns FAccountIds to the passed in arguments for three accounts */
-	FTestPipeline& GetLoginPipeline(FAccountId& OutAccountId, FAccountId& OutAccountId2, FAccountId& OutAccountId3) const;
-	/* Creates a login pipeline and logs in user, immediately assigns FAccountIds to the passed in arguments for four accounts */
-	FTestPipeline& GetLoginPipeline(FAccountId& OutAccountId, FAccountId& OutAccountId2, FAccountId& OutAccountId3, FAccountId& OutAccountId4) const;
-	/* Creates a login pipeline and logs in user, immediately assigns FAccountIds to the passed in arguments for five accounts */
-	FTestPipeline& GetLoginPipeline(FAccountId& OutAccountId, FAccountId& OutAccountId2, FAccountId& OutAccountId3, FAccountId& OutAccountId4, FAccountId& OutAccountId5) const;
 
 	/* Attempts to assgin OutAccountId to what LocalUserId is logged in as. */
 	void AssignLoginUsers(int32 LocalUserId, FAccountId& OutAccountId) const;
@@ -63,14 +57,14 @@ protected:
 
 	OnlineTestBase();
 
-	/* Helper function to delete current accounts on the TestDataService and wait for it to finish */
-	bool DeleteAccountsForCurrentTemplate() const;
+	/* Helper function to delete accounts from TestDataService */
+	bool DeleteAccounts(int32 TemplateRequestedNum = -1) const;
 
 	/* Helper function to destroy the current OnlineService module which stores some state that maybe needs to be reset */
 	void DestroyCurrentServiceModule() const;
 
 	/* Proxy function to the DestroyCurrentServiceModule() */
-	bool ResetAccountStatus() const;
+	bool ResetAccountStatus(int32 TemplateRequestedNum = -1) const;
 
 	/* Returns the name of the service we're currently testing */
 	FString GetService() const;
@@ -80,25 +74,26 @@ protected:
 // If you define a custom credentials path you must define these in your own Auth.cpp file
 #if ONLINETESTS_USEEXTERNAUTH
 	// This will provide a paramters to pass to the OSS Auth interface.
-	UE::Online::FAuthLogin::Params CustomCredentials(int LocalUserNum) const;
+	TArray<UE::Online::FAuthLogin::Params> CustomCredentials(int32 LocalUserNum, int32 NumUsers) const;
 
 	// When called should "reset" accounts to default state for Stats, Achivements, Lobbies and other
 	// stateful interfaces to work.
-	bool CustomResetAccounts() const;
+	bool CustomResetAccounts(int32 TemplateRequestedNum) const;
 
 	// When called will *remove* the account from the underlying service.
 	// This is generally called to test new-user login or first-time setup.
 	// To be called rare-ly.
-	bool CustomDeleteAccounts() const;
+	bool CustomDeleteAccounts(int32 TemplateRequestedNum) const;
+
 #endif
+	TArray<UE::Online::FAuthLogin::Params> GetIniCredentials(int32 LocalUserNum) const;
 
-	UE::Online::FAuthLogin::Params GetIniCredentials(int LocalUserNum) const;
-	UE::Online::FAuthLogin::Params GetCredentials(int LocalUserNum) const;
+	TArray<UE::Online::FAuthLogin::Params> GetCredentials(int32 LocalUserNum, int32 NumUsers) const;
 
 	/* Returns the ini login category name for the configured service */
 	FString GetLoginCredentialCategory() const;
 
-	void RunToCompletion(bool bLogout = true, const TOptional<int32> UserNumToLogout = TOptional<int32>()) const;
+	void RunToCompletion(bool bLogout = true) const;
 
 	/* ITestInvoker */
 	virtual void invoke() const override = 0;
@@ -111,6 +106,7 @@ private:
 	mutable FTestDriver Driver;
 	mutable TSharedPtr<FTestPipeline> Pipeline;
 	mutable uint32 NumLocalUsers = -1;
+	mutable uint32 UserNumToLogout = -1;
 };
 
 typedef OnlineTestBase* (*OnlineTestConstructor)();
