Compile errors in AccountServiceTest.cpp

diff --git a/Engine/Source/Programs/Online/OSSTestsCore/Tests/AccountServiceTest.cpp b/Engine/Source/Programs/Online/OSSTestsCore/Tests/AccountServiceTest.cpp
--- a/Engine/Source/Programs/Online/OSSTestsCore/Tests/AccountServiceTest.cpp
+++ b/Engine/Source/Programs/Online/OSSTestsCore/Tests/AccountServiceTest.cpp
@@ -23,9 +23,8 @@ ACCOUNTSERVICE_TEST_CASE("Verify if we can properly create a OnlineAccountCreden
 
 ACCOUNTSERVICE_TEST_CASE("Verify if we can properly instantiate the OSS")
 {
-	FTestPipeline& Pipeline = GetPipeline();
-
-	Pipeline.EmplaceLambda([](IOnlineSubsystem* Services)
+	GetPipeline()
+		.EmplaceLambda([](IOnlineSubsystem* Services)
 		{
 			int32 LocalUserNum = 0;
 			FString LocalType = "test_account";
