Added Commerce Query/Get Entitlements tests.

diff --git a/Engine/Source/Programs/Online/OnlineTestsCore/Source/OnlineTestsCore.Build.cs b/Engine/Source/Programs/Online/OnlineTestsCore/Source/OnlineTestsCore.Build.cs
--- a/Engine/Source/Programs/Online/OnlineTestsCore/Source/OnlineTestsCore.Build.cs
+++ b/Engine/Source/Programs/Online/OnlineTestsCore/Source/OnlineTestsCore.Build.cs
@@ -31,7 +31,8 @@ public class OnlineTestsCore : ModuleRules
 				"OnlineServicesOSSAdapter",
 				"SSL",
 				"Json",
-				"JsonUtilities"
+				"JsonUtilities",
+				"HTTP"
 			}
 		);
 
@@ -40,6 +41,12 @@ public class OnlineTestsCore : ModuleRules
 		{
 			PublicDefinitions.Add(String.Format("ONLINETESTS_USEEXTERNAUTH=0"));
 		}
+
+		PrivateIncludePaths.AddRange(
+			new string[] {
+				Path.Combine(Target.RelativeEnginePath, "Restricted", "NotForLicensees", "Source", "Programs", "OnlineTests", "Private")
+			}
+		);
 	}
 }
 
diff --git a/Engine/Source/Programs/Online/OnlineTestsCore/Source/Private/Tests/Commerce/GetEntitlementsTests.cpp b/Engine/Source/Programs/Online/OnlineTestsCore/Source/Private/Tests/Commerce/GetEntitlementsTests.cpp
deleted file mode 100644
diff --git a/Engine/Source/Programs/Online/OnlineTestsCore/Source/Private/Tests/Commerce/QueryEntitlementsTests.cpp b/Engine/Source/Programs/Online/OnlineTestsCore/Source/Private/Tests/Commerce/QueryEntitlementsTests.cpp
deleted file mode 100644
diff --git a/Engine/Source/Programs/Online/OnlineTestsCore/Source/Public/Helpers/Auth/AuthQueryExternalAuthToken.h b/Engine/Source/Programs/Online/OnlineTestsCore/Source/Public/Helpers/Auth/AuthQueryExternalAuthToken.h
new file mode 100644
--- /dev/null
+++ b/Engine/Source/Programs/Online/OnlineTestsCore/Source/Public/Helpers/Auth/AuthQueryExternalAuthToken.h
@@ -0,0 +1,53 @@
+// Copyright Epic Games, Inc. All Rights Reserved.
+#pragma once
+
+#include "CoreMinimal.h"
+#include "TestDriver.h"
+#include "TestHarness.h"
+
+#include "Online/AuthCommon.h"
+#include "Online/Auth.h"
+#include "Online/OnlineAsyncOp.h"
+#include "Online/OnlineErrorDefinitions.h"
+#include "Online/OnlineServicesCommon.h"
+#include "OnlineCatchHelper.h"
+
+using namespace UE::Online;
+
+struct FAuthQueryExternalAuthTokenStep : public FAsyncTestStep
+{
+	FAuthQueryExternalAuthTokenStep(UE::Online::FAuthQueryExternalAuthToken::Params&& InQueryParams)
+		:	QueryParams(MoveTemp(InQueryParams))
+	{
+	}
+
+	virtual ~FAuthQueryExternalAuthTokenStep()
+	{
+		if (OnlineAuthPtr != nullptr)
+		{
+			OnlineAuthPtr = nullptr;
+		}
+	};
+
+	virtual void Run(FAsyncStepResult Promise, SubsystemType Services) override
+	{
+		OnlineAuthPtr = Services->GetAuthInterface();
+		REQUIRE(OnlineAuthPtr);
+
+		OnlineAuthPtr->QueryExternalAuthToken(MoveTemp(QueryParams))
+			.OnComplete([this, Promise = MoveTemp(Promise)](const TOnlineResult<FAuthQueryExternalAuthToken>& Op) mutable
+				{
+					REQUIRE_OP(Op);
+					const UE::Online::FExternalAuthToken ExternalAuthToken = Op.GetOkValue().ExternalAuthToken;
+					REQUIRE(ExternalAuthToken.Type.IsValid());
+					REQUIRE(!ExternalAuthToken.Data.IsEmpty());
+
+					Promise->SetValue(true);
+				});
+	}
+
+protected:
+	int32 LocalUserNum = 0;
+	UE::Online::FAuthQueryExternalAuthToken::Params QueryParams;
+	UE::Online::IAuthPtr OnlineAuthPtr = nullptr;
+};
\ No newline at end of file
diff --git a/Engine/Source/Programs/Online/OnlineTestsCore/Source/Public/OnlineCatchHelper.h b/Engine/Source/Programs/Online/OnlineTestsCore/Source/Public/OnlineCatchHelper.h
--- a/Engine/Source/Programs/Online/OnlineTestsCore/Source/Public/OnlineCatchHelper.h
+++ b/Engine/Source/Programs/Online/OnlineTestsCore/Source/Public/OnlineCatchHelper.h
@@ -52,6 +52,10 @@ public:
 	/* Attempts to assgin OutAccountId to what LocalUserId is logged in as. */
 	void AssignLoginUsers(int32 LocalUserId, FAccountId& OutAccountId) const;
 
+#if ONLINETESTS_USEEXTERNAUTH
+	TArray<FString> GetExternalAuthAccountIds(int32 TemplateRequestedNum) const;
+#endif
+
 	FTestPipeline& GetPipeline() const;
 protected:
 
