Added Commerce RedeemEntitlements tests.

diff --git a/Engine/Plugins/Online/OnlineServices/Source/OnlineServicesInterface/Public/Online/OnlineErrorDefinitions.h b/Engine/Plugins/Online/OnlineServices/Source/OnlineServicesInterface/Public/Online/OnlineErrorDefinitions.h
--- a/Engine/Plugins/Online/OnlineServices/Source/OnlineServicesInterface/Public/Online/OnlineErrorDefinitions.h
+++ b/Engine/Plugins/Online/OnlineServices/Source/OnlineServicesInterface/Public/Online/OnlineErrorDefinitions.h
@@ -33,6 +33,7 @@ UE_ONLINE_ERROR_COMMON(Common, WillRetry, 19, TEXT("will_retry"), LOCTEXT("WillR
 UE_ONLINE_ERROR_COMMON(Common, Timeout, 20, TEXT("timeout"), LOCTEXT("Timeout", "Operation timed out"))
 UE_ONLINE_ERROR_COMMON(Common, InvalidState, 21, TEXT("invalid_state"), LOCTEXT("InvalidState", "Invalid state"))
 UE_ONLINE_ERROR_COMMON(Common, Unknown, 22, TEXT("unknown"), LOCTEXT("Unknown", "Unknown Error"))
+UE_ONLINE_ERROR_COMMON(Common, NoChange, 23, TEXT("no_change"), LOCTEXT("NoChange", "No change"))
 
 } /* namespace UE::Online::Errors */
 
diff --git a/Engine/Plugins/Online/OnlineServicesEOS/Source/Private/Online/CommerceEOS.cpp b/Engine/Plugins/Online/OnlineServicesEOS/Source/Private/Online/CommerceEOS.cpp
--- a/Engine/Plugins/Online/OnlineServicesEOS/Source/Private/Online/CommerceEOS.cpp
+++ b/Engine/Plugins/Online/OnlineServicesEOS/Source/Private/Online/CommerceEOS.cpp
@@ -436,15 +436,21 @@ TOnlineAsyncOpHandle<FCommerceRedeemEntitlement> FCommerceEOS::RedeemEntitlement
 		const FCommerceRedeemEntitlement::Params& Params = Op.GetParams();
 		if (!Services.Get<FAuthEOS>()->IsLoggedIn(Params.LocalAccountId))
 		{
-			Promise.SetValue(nullptr);
 			Op.SetError(Errors::NotLoggedIn());
+			Promise.SetValue(nullptr);
 			return;
 		}
 		EOS_EpicAccountId LocalUserEasId = GetEpicAccountId(Params.LocalAccountId);
 		if (!EOS_EpicAccountId_IsValid(LocalUserEasId))
 		{
-			Promise.SetValue(nullptr);
 			Op.SetError(Errors::NotLoggedIn());
+			Promise.SetValue(nullptr);
+			return;
+		}
+		if (Params.Quantity != 1)
+		{
+			Op.SetError(Errors::InvalidParams());
+			Promise.SetValue(nullptr);
 			return;
 		}
 
diff --git a/Engine/Plugins/Online/OnlineServicesEOSGS/Source/Private/Online/OnlineErrorEOSGS.cpp b/Engine/Plugins/Online/OnlineServicesEOSGS/Source/Private/Online/OnlineErrorEOSGS.cpp
--- a/Engine/Plugins/Online/OnlineServicesEOSGS/Source/Private/Online/OnlineErrorEOSGS.cpp
+++ b/Engine/Plugins/Online/OnlineServicesEOSGS/Source/Private/Online/OnlineErrorEOSGS.cpp
@@ -65,6 +65,7 @@ FOnlineError MapCommonEOSError(FOnlineError&& Error, EOS_EResult Result)
 	case EOS_EResult::EOS_Canceled:				return Errors::Cancelled(MoveTemp(Error));
 	case EOS_EResult::EOS_NotFound:				return Errors::NotFound(MoveTemp(Error));
 	case EOS_EResult::EOS_OperationWillRetry:	return Errors::WillRetry(MoveTemp(Error));
+	case EOS_EResult::EOS_NoChange:				return Errors::NoChange(MoveTemp(Error));
 	case EOS_EResult::EOS_VersionMismatch:		return Errors::IncompatibleVersion(MoveTemp(Error));
 	case EOS_EResult::EOS_LimitExceeded:		return Errors::TooManyRequests(MoveTemp(Error));
 	case EOS_EResult::EOS_TimedOut:				return Errors::Timeout(MoveTemp(Error));
diff --git a/Engine/Source/Programs/Online/OnlineTestsCore/Source/Private/Tests/Commerce/RedeemEntitlementTests.cpp b/Engine/Source/Programs/Online/OnlineTestsCore/Source/Private/Tests/Commerce/RedeemEntitlementTests.cpp
deleted file mode 100644
diff --git a/Engine/Source/Programs/Online/OnlineTestsCore/Source/Public/Helpers/Commerce/RedeemEntitlementsHelper.h b/Engine/Source/Programs/Online/OnlineTestsCore/Source/Public/Helpers/Commerce/RedeemEntitlementsHelper.h
new file mode 100644
--- /dev/null
+++ b/Engine/Source/Programs/Online/OnlineTestsCore/Source/Public/Helpers/Commerce/RedeemEntitlementsHelper.h
@@ -0,0 +1,66 @@
+// Copyright Epic Games, Inc. All Rights Reserved.
+
+#pragma once
+
+#include "Online/CommerceCommon.h"
+#include "AsyncTestStep.h"
+#include "OnlineCatchHelper.h"
+#include "Online/OnlineServicesLog.h"
+
+struct FRedeemEntitlementsHelper : public FAsyncTestStep
+{
+	using ParamsType = UE::Online::FCommerceRedeemEntitlement::Params;
+	using ResultType = UE::Online::TOnlineResult<UE::Online::FCommerceRedeemEntitlement>;
+
+	struct FHelperParams
+	{
+		ParamsType* OpParams = nullptr;
+		TOptional<ResultType> ExpectedError;
+	};
+
+	FRedeemEntitlementsHelper(FHelperParams&& InHelperParams, TSharedPtr<FString> InEntitlementId)
+		:	HelperParams(MoveTemp(InHelperParams))
+		,	EntitlementId(InEntitlementId)
+	{
+		REQUIRE(HelperParams.OpParams);
+		REQUIRE((!HelperParams.ExpectedError.IsSet() || HelperParams.ExpectedError->IsError()));
+	}
+
+	virtual ~FRedeemEntitlementsHelper() = default;
+
+	virtual void Run(FAsyncStepResult Promise, SubsystemType Services) override
+	{
+		HelperParams.OpParams->EntitlementId = *EntitlementId;
+
+		CommerceInterface = Services->GetCommerceInterface();
+		REQUIRE(CommerceInterface);
+
+		if (!HelperParams.ExpectedError.IsSet())
+		{
+			CommerceInterface->RedeemEntitlement(MoveTemp(*HelperParams.OpParams))
+				.OnComplete([this, Promise = MoveTemp(Promise)](const ResultType& Result)
+					{
+						REQUIRE(Result.IsOk());
+						Promise->SetValue(true);
+					});
+		}
+		else
+		{
+			ELogVerbosity::Type OldVerbosity = LogOnlineServices.GetVerbosity();
+			LogOnlineServices.SetVerbosity(ELogVerbosity::NoLogging);
+			CommerceInterface->RedeemEntitlement(MoveTemp(*HelperParams.OpParams))
+				.OnComplete([this, OldVerbosity, Promise = MoveTemp(Promise)](const ResultType& Result)
+					{
+						REQUIRE_OP_EQ(Result, HelperParams.ExpectedError->GetErrorValue());
+						LogOnlineServices.SetVerbosity(OldVerbosity);
+						Promise->SetValue(true);
+					});
+		}
+
+	}
+
+protected:
+	FHelperParams HelperParams;
+	TSharedPtr<FString> EntitlementId;
+	UE::Online::ICommercePtr CommerceInterface = nullptr;
+};
\ No newline at end of file
