Destroy online servers after achievement stats tests in addition to before them. Ensure that the stats values in OnlineServices do not impact unrelated future tests

diff --git a/Engine/Source/Programs/Online/OnlineTestsCore/Source/Private/Tests/Achievements/UnlockAchievementsTests.cpp b/Engine/Source/Programs/Online/OnlineTestsCore/Source/Private/Tests/Achievements/UnlockAchievementsTests.cpp
--- a/Engine/Source/Programs/Online/OnlineTestsCore/Source/Private/Tests/Achievements/UnlockAchievementsTests.cpp
+++ b/Engine/Source/Programs/Online/OnlineTestsCore/Source/Private/Tests/Achievements/UnlockAchievementsTests.cpp
@@ -196,6 +196,7 @@ ACHIEVEMENTS_TEST_CASE("Unlock Title-Managed Achievements By Stats", EG_ACHIEVEM
 			});
 
 	RunToCompletion();
+	DestroyCurrentServiceModule();
 }
 
 ACHIEVEMENTS_TEST_CASE("Unlock Title-Managed Achievements By Stats With Multiple Conditions", EG_ACHIEVEMENTS_UNLOCKACHIEVEMENTS_TITLEMANAGED_TAG)
@@ -249,4 +250,5 @@ ACHIEVEMENTS_TEST_CASE("Unlock Title-Managed Achievements By Stats With Multiple
 			});
 
 	RunToCompletion();
+	DestroyCurrentServiceModule();
 }
\ No newline at end of file
