OSSTests: integrate TDS support to use accounts for EOS OSSv1 testing and resolve identity issues related to this.

diff --git a/Engine/Source/Programs/Online/OSSTestsCore/OSSTestsCore.Build.cs b/Engine/Source/Programs/Online/OSSTestsCore/OSSTestsCore.Build.cs
--- a/Engine/Source/Programs/Online/OSSTestsCore/OSSTestsCore.Build.cs
+++ b/Engine/Source/Programs/Online/OSSTestsCore/OSSTestsCore.Build.cs
@@ -2,6 +2,8 @@
 
 using UnrealBuildTool;
 using System.IO;
+using System.Linq;
+using System;
 
 public class OSSTestsCore : ModuleRules
 {
@@ -22,11 +24,10 @@ public class OSSTestsCore : ModuleRules
 			}
 		);
 
-		PrivateIncludePaths.AddRange(
-			new string[] {
-				Path.Combine(Target.RelativeEnginePath, "Plugins", "Online"),
-				Path.Combine(Target.RelativeEnginePath, "Restricted", "NotForLicensees", "Plugins", "Online")
-            }
-        );
+		// Disable external auth if target doesn't define it.
+		if (!Target.GlobalDefinitions.Contains("OSSTESTS_USEEXTERNAUTH=1"))
+		{
+			PublicDefinitions.Add(String.Format("OSSTESTS_USEEXTERNAUTH=0"));
+		}
 	}
 }
diff --git a/Engine/Source/Programs/Online/OSSTestsCore/OnlineSubsystemCatchHelper.cpp b/Engine/Source/Programs/Online/OSSTestsCore/OnlineSubsystemCatchHelper.cpp
--- a/Engine/Source/Programs/Online/OSSTestsCore/OnlineSubsystemCatchHelper.cpp
+++ b/Engine/Source/Programs/Online/OSSTestsCore/OnlineSubsystemCatchHelper.cpp
@@ -143,31 +143,51 @@ FString OnlineSubsystemTestBase::GetSubsystem() const
 	return Subsystem;
 }
 
-FOnlineAccountCredentials OnlineSubsystemTestBase::GetCredentials(int LocalUserNum) const
+TArray<FOnlineAccountCredentials> OnlineSubsystemTestBase::GetIniCredentials(int32 NumUsers) const
 {
 	FString LoginCredentialCategory = FString::Printf(TEXT("LoginCredentials %s"), *Subsystem);
 	TArray<FString> CredentialsArr;
 	GConfig->GetArray(*LoginCredentialCategory, TEXT("Credentials"), CredentialsArr, GEngineIni);
 
-	if (LocalUserNum > CredentialsArr.Num())
+	if (NumUsers > CredentialsArr.Num())
 	{
 		UE_LOG(LogOSSTests, Error, TEXT("Attempted to GetCredentials for more than we have stored! Add more credentials to the DefaultEngine.ini for OssTests"));
-		return FOnlineAccountCredentials(FString(TEXT("")), FString(TEXT("")), FString(TEXT("")));
+		return TArray<FOnlineAccountCredentials>();
 	}
-	
-	FString LoginUsername, LoginType, LoginPassword;
-	FParse::Value(*CredentialsArr[LocalUserNum], TEXT("Type="), LoginType);
-	FParse::Value(*CredentialsArr[LocalUserNum], TEXT("Id="), LoginUsername);
-	FParse::Value(*CredentialsArr[LocalUserNum], TEXT("Password="), LoginPassword);
 
-	// Return an account, using the default first user if we only have one account for the subsystem.
-	return FOnlineAccountCredentials(LoginType, LoginUsername, LoginPassword);
+	TArray<FOnlineAccountCredentials> OnlineAccountCredentials;
+	for (int32 Index = 0; Index < CredentialsArr.Num(); ++Index)
+	{
+		FString LoginUsername, LoginType, LoginToken;
+		FParse::Value(*CredentialsArr[Index], TEXT("Type="), LoginType);
+		FParse::Value(*CredentialsArr[Index], TEXT("Id="), LoginUsername);
+		FParse::Value(*CredentialsArr[Index], TEXT("Token="), LoginToken);
+		INFO(*FString::Printf(TEXT("Logging in with type %s, id %s, password %s"), *LoginType, *LoginUsername, *LoginToken));
+
+		OnlineAccountCredentials.Add(FOnlineAccountCredentials{ LoginType, LoginUsername, LoginToken });
+	}
+	return OnlineAccountCredentials;
+}
+
+TArray<FOnlineAccountCredentials> OnlineSubsystemTestBase::GetCredentials(int32 LocalUserNum, int32 NumUsers) const
+{
+#if OSSTESTS_USEEXTERNAUTH
+	return CustomCredentials(LocalUserNum, NumUsers);
+#else // OSSTESTS_USEEXTERNAUTH
+	return GetIniCredentials(LocalUserNum);
+#endif // OSSTESTS_USEEXTERNAUTH
+}
+
+FString OnlineSubsystemTestBase::GetLoginCredentialCategory() const
+{
+	return FString::Printf(TEXT("LoginCredentials %s"), *Subsystem);
 }
 
 FTestPipeline& OnlineSubsystemTestBase::GetLoginPipeline(uint32 NumUsersToLogin) const
 {
 	REQUIRE(NumLocalUsers == -1); // Don't call GetLoginPipeline more than once per test
 	NumLocalUsers = NumUsersToLogin;
+	NumUsersToLogout = NumUsersToLogin;
 
 	bool bUseAutoLogin = false;
 	bool bUseImplicitLogin = false;
@@ -189,9 +209,11 @@ FTestPipeline& OnlineSubsystemTestBase::GetLoginPipeline(uint32 NumUsersToLogin)
 	}
 	else
 	{
-		for (uint32 i = 0; i < NumUsersToLogin; i++)
+		TArray<FOnlineAccountCredentials> AuthLoginParams = GetCredentials(0, NumUsersToLogin);
+
+		for (uint32 Index = 0; Index < NumUsersToLogin; ++Index)
 		{
-			Pipeline.EmplaceStep<FIdentityLoginStep>(i, GetCredentials(i));
+			Pipeline.EmplaceStep<FIdentityLoginStep>(Index, AuthLoginParams[Index]);
 		}
 	}
 
diff --git a/Engine/Source/Programs/Online/OSSTestsCore/OnlineSubsystemCatchHelper.h b/Engine/Source/Programs/Online/OSSTestsCore/OnlineSubsystemCatchHelper.h
--- a/Engine/Source/Programs/Online/OSSTestsCore/OnlineSubsystemCatchHelper.h
+++ b/Engine/Source/Programs/Online/OSSTestsCore/OnlineSubsystemCatchHelper.h
@@ -4,9 +4,11 @@
 
 #include "TestHarness.h"
 #include "TestDriver.h"
+#include "CoreMinimal.h"
 #include "Interfaces/OnlineIdentityInterface.h"
 #include "Misc/ConfigCacheIni.h"
 #include "Modules/ModuleManager.h"
+#include "Containers/Ticker.h"
 
 #include <catch2/catch_test_case_info.hpp>
 #include <catch2/interfaces/catch_interfaces_registry_hub.hpp>
@@ -28,9 +30,21 @@ protected:
 	OnlineSubsystemTestBase();
 	
 	FString GetSubsystem() const;
-	FOnlineAccountCredentials GetCredentials(int LocalUserNum) const;
+
+#if OSSTESTS_USEEXTERNAUTH
+	TArray<FOnlineAccountCredentials> CustomCredentials(int32 LocalUserNum, int32 NumUsers) const;
+#endif // OSSTESTS_USEEXTERNAUTH
+
+	TArray<FOnlineAccountCredentials> GetIniCredentials(int32 NumUsers) const;
+
+	TArray<FOnlineAccountCredentials> GetCredentials(int32 LocalUserNum, int32 NumUsers) const;
+
 	FTestPipeline& GetLoginPipeline(uint32 NumUsersToLogin = 1) const;
 	FTestPipeline& GetPipeline() const;
+
+	/* Returns the ini login category name for the configured service */
+	FString GetLoginCredentialCategory() const;
+
 	void RunToCompletion() const;
 
 	/* ITestInvoker */
@@ -43,6 +57,7 @@ private:
 	mutable FTestDriver Driver;
 	mutable FTestPipeline Pipeline;
 	mutable uint32 NumLocalUsers = -1;
+	mutable uint32 NumUsersToLogout = -1;
 };
 
 typedef OnlineSubsystemTestBase* (*OnlineSubsystemTestConstructor)();
diff --git a/Engine/Source/Programs/Online/OSSTestsCore/TestDriver.h b/Engine/Source/Programs/Online/OSSTestsCore/TestDriver.h
--- a/Engine/Source/Programs/Online/OSSTestsCore/TestDriver.h
+++ b/Engine/Source/Programs/Online/OSSTestsCore/TestDriver.h
@@ -12,7 +12,7 @@ DEFINE_LOG_CATEGORY_STATIC(LogOSSTests, Log, Log);
 
 #define UE_LOG_OSSTESTS(Verbosity, Format, ...) \
 { \
-	UE_LOG(Log, Verbosity, "OSSTests: ", *FString::Printf(Format, ##__VA_ARGS__)); \
+	UE_LOG(LogOSSTests, Verbosity, TEXT("OSSTests: %s"), *FString::Printf(Format, ##__VA_ARGS__)); \
 }
 
 static const FTimespan TICK_DURATION = FTimespan::FromMilliseconds(1);
diff --git a/Engine/Source/Programs/Online/OSSTestsCore/Tests/AccountServiceTest.cpp b/Engine/Source/Programs/Online/OSSTestsCore/Tests/AccountServiceTest.cpp
--- a/Engine/Source/Programs/Online/OSSTestsCore/Tests/AccountServiceTest.cpp
+++ b/Engine/Source/Programs/Online/OSSTestsCore/Tests/AccountServiceTest.cpp
@@ -24,16 +24,14 @@ ACCOUNTSERVICE_TEST_CASE("Verify if we can properly create a OnlineAccountCreden
 ACCOUNTSERVICE_TEST_CASE("Verify if we can properly instantiate the OSS")
 {
 	GetPipeline()
-		.EmplaceLambda([](IOnlineSubsystem* Services)
+		.EmplaceLambda([this](IOnlineSubsystem* Services)
 		{
 			int32 LocalUserNum = 0;
-			FString LocalType = "test_account";
-			FString LocalId = "12345";
-			FString LocalToken = "fake_token";
-			FOnlineAccountCredentials LocalAccount = FOnlineAccountCredentials(LocalType, LocalId, LocalToken);
+			int32 NumUsers = 1;
+			TArray<FOnlineAccountCredentials> AccountCreds = GetCredentials(LocalUserNum, NumUsers);
 
 			IOnlineIdentityPtr IdentityInterface = Services->GetIdentityInterface();		
-			bool LoggedIn = IdentityInterface.Get()->Login(LocalUserNum, LocalAccount);
+			bool LoggedIn = IdentityInterface.Get()->Login(LocalUserNum, AccountCreds[0]);
 			REQUIRE(LoggedIn);
 		});
 
diff --git a/Engine/Source/Programs/Online/OSSTestsCore/Tests/Identity/IdentityAutoLoginTests.cpp b/Engine/Source/Programs/Online/OSSTestsCore/Tests/Identity/IdentityAutoLoginTests.cpp
--- a/Engine/Source/Programs/Online/OSSTestsCore/Tests/Identity/IdentityAutoLoginTests.cpp
+++ b/Engine/Source/Programs/Online/OSSTestsCore/Tests/Identity/IdentityAutoLoginTests.cpp
@@ -16,12 +16,13 @@ IDENTITY_TEST_CASE("Verify calling Identity AutoLogin with valid inputs returns
 { 
 	FTestDriver LocalDriver;
 	int32 LocalUserNum = 0;
+	int32 NumUsers = 1;
 	
-	FOnlineAccountCredentials AccountCreds = GetCredentials(LocalUserNum);
+	TArray<FOnlineAccountCredentials> AccountCreds = GetCredentials(LocalUserNum, NumUsers);
 
-	FString LoginCredentialsType = TEXT("AUTH_TYPE=") + AccountCreds.Type + ",";
-	FString LoginCredentialsId = TEXT("AUTH_LOGIN=") + AccountCreds.Id + + ",";
-	FString LoginCredentialsPassword = TEXT("AUTH_PASSWORD=") + AccountCreds.Token;
+	FString LoginCredentialsType = TEXT("AUTH_TYPE=") + AccountCreds[0].Type + ",";
+	FString LoginCredentialsId = TEXT("AUTH_LOGIN=") + AccountCreds[0].Id + + ",";
+	FString LoginCredentialsPassword = TEXT("AUTH_PASSWORD=") + AccountCreds[0].Token;
 
 	FCommandLine::Set(*(LoginCredentialsType + LoginCredentialsId + LoginCredentialsPassword));
 
diff --git a/Engine/Source/Programs/Online/OSSTestsCore/Tests/Identity/IdentityGetLoginStatusTests.cpp b/Engine/Source/Programs/Online/OSSTestsCore/Tests/Identity/IdentityGetLoginStatusTests.cpp
--- a/Engine/Source/Programs/Online/OSSTestsCore/Tests/Identity/IdentityGetLoginStatusTests.cpp
+++ b/Engine/Source/Programs/Online/OSSTestsCore/Tests/Identity/IdentityGetLoginStatusTests.cpp
@@ -23,9 +23,11 @@ IDENTITY_TEST_CASE("Verify calling Identity GetLoginStatus by UserId with valid
 	ELoginStatus::Type ExpectedLoginStatus = ELoginStatus::Type::LoggedIn;
 	ELoginStatus::Type ExpectedLogoutStatus = ELoginStatus::Type::NotLoggedIn;
 	int32 NumUsersToImplicitLogin = 0;
+	int32 NumUsers = 1;
+	TArray<FOnlineAccountCredentials> AccountCreds = GetCredentials(LocalUserNum, NumUsers);
 
 	FTestPipeline LocalPipeline = LocalDriver.MakePipeline()
-		.EmplaceStep<FIdentityLoginStep>(LocalUserNum, GetCredentials(LocalUserNum))
+		.EmplaceStep<FIdentityLoginStep>(LocalUserNum, AccountCreds[0])
 		.EmplaceStep<FIdentityGetUniquePlayerIdStep>(LocalUserNum, [&LocalUserId](FUniqueNetIdPtr InUserId) {LocalUserId = InUserId; })
 		.EmplaceStep<FIdentityGetLoginStatusByUserIdStep>(&LocalUserId, ExpectedLoginStatus)
 		.EmplaceStep<FIdentityLogoutStep>(LocalUserNum)
@@ -42,9 +44,11 @@ IDENTITY_TEST_CASE("Verify calling Identity GetLoginStatus by LocalUserNum with
 	ELoginStatus::Type ExpectedLoginStatus = ELoginStatus::Type::LoggedIn;
 	ELoginStatus::Type ExpectedLogoutStatus = ELoginStatus::Type::NotLoggedIn;
 	int32 NumUsersToImplicitLogin = 0;
+	int32 NumUsers = 1;
+	TArray<FOnlineAccountCredentials> AccountCreds = GetCredentials(LocalUserNum, NumUsers);
 
 	FTestPipeline LocalPipeline = LocalDriver.MakePipeline()
-		.EmplaceStep<FIdentityLoginStep>(LocalUserNum, GetCredentials(LocalUserNum))
+		.EmplaceStep<FIdentityLoginStep>(LocalUserNum, AccountCreds[0])
 		.EmplaceStep<FIdentityGetLoginStatusByLocalUserNumStep>(LocalUserNum, ExpectedLoginStatus)
 		.EmplaceStep<FIdentityLogoutStep>(LocalUserNum)
 		.EmplaceStep<FIdentityGetLoginStatusByLocalUserNumStep>(LocalUserNum, ExpectedLogoutStatus);
diff --git a/Engine/Source/Programs/Online/OSSTestsCore/Tests/Identity/IdentityLogoutTests.cpp b/Engine/Source/Programs/Online/OSSTestsCore/Tests/Identity/IdentityLogoutTests.cpp
--- a/Engine/Source/Programs/Online/OSSTestsCore/Tests/Identity/IdentityLogoutTests.cpp
+++ b/Engine/Source/Programs/Online/OSSTestsCore/Tests/Identity/IdentityLogoutTests.cpp
@@ -14,12 +14,12 @@
 IDENTITY_TEST_CASE("Verify calling Identity Logout with valid inputs returns the expected result(Success Case)", EG_IDENTITY_LOGOUT_TAG)
 {
 	int32 LocalUserNum = 0;
-	int32 TargetUserNum = 1;
-	int32 NumUsersToImplicitLogin = 1;
+	int32 NumUsers = 1;
+	TArray<FOnlineAccountCredentials> AccountCreds = GetCredentials(LocalUserNum, NumUsers);
 
-	GetLoginPipeline(NumUsersToImplicitLogin)
-		.EmplaceStep<FIdentityLoginStep>(TargetUserNum, GetCredentials(TargetUserNum))
-		.EmplaceStep<FIdentityLogoutStep>(TargetUserNum);
+	GetPipeline()
+		.EmplaceStep<FIdentityLoginStep>(LocalUserNum, AccountCreds[0])
+		.EmplaceStep<FIdentityLogoutStep>(LocalUserNum);
 
 	RunToCompletion();
 }
