Use default services config instance
#tests [Null] test suite

diff --git a/Engine/Source/Programs/Online/OnlineTestsCore/Source/Private/OnlineCatchHelper.cpp b/Engine/Source/Programs/Online/OnlineTestsCore/Source/Private/OnlineCatchHelper.cpp
--- a/Engine/Source/Programs/Online/OnlineTestsCore/Source/Private/OnlineCatchHelper.cpp
+++ b/Engine/Source/Programs/Online/OnlineTestsCore/Source/Private/OnlineCatchHelper.cpp
@@ -76,7 +76,7 @@ EOnlineServices OnlineTestBase::GetServiceType() const
 
 SubsystemType OnlineTestBase::GetSubsystem() const
 {
-	return UE::Online::FOnlineServicesRegistry::Get().GetNamedServicesInstance(ServiceType, NAME_None);
+	return UE::Online::FOnlineServicesRegistry::Get().GetNamedServicesInstance(ServiceType, NAME_None, NAME_None);
 }
 
 bool OnlineTestBase::DeleteAccounts(int32 TemplateRequestedNum) const
@@ -90,7 +90,7 @@ bool OnlineTestBase::DeleteAccounts(int32 TemplateRequestedNum) const
 
 void OnlineTestBase::DestroyCurrentServiceModule() const
 {
-	UE::Online::FOnlineServicesRegistry::Get().DestroyNamedServicesInstance(ServiceType, NAME_None);
+	UE::Online::FOnlineServicesRegistry::Get().DestroyNamedServicesInstance(ServiceType, NAME_None, NAME_None);
 }
 
 bool OnlineTestBase::ResetAccountStatus(int32 TemplateRequestedNum) const
diff --git a/Engine/Source/Programs/Online/OnlineTestsCore/Source/Private/TestDriver.cpp b/Engine/Source/Programs/Online/OnlineTestsCore/Source/Private/TestDriver.cpp
--- a/Engine/Source/Programs/Online/OnlineTestsCore/Source/Private/TestDriver.cpp
+++ b/Engine/Source/Programs/Online/OnlineTestsCore/Source/Private/TestDriver.cpp
@@ -162,7 +162,7 @@ FTestDriver::~FTestDriver()
 
 bool FTestDriver::AddPipeline(FTestPipeline&& Pipeline, const FPipelineTestContext& TestContext)
 {
-	SubsystemType OnlineSubsystem = UE::Online::FOnlineServicesRegistry::Get().GetNamedServicesInstance(TestContext.InitOptions.ServiceType, NAME_None);
+	SubsystemType OnlineSubsystem = UE::Online::FOnlineServicesRegistry::Get().GetNamedServicesInstance(TestContext.InitOptions.ServiceType, NAME_None, NAME_None);
 	if (OnlineSubsystem == nullptr)
 	{
 		return false;
