Allow 3rd party identity provider authentication with EOS platform

diff --git a/Engine/Source/Programs/Online/OSSTestsCore/OnlineSubsystemCatchHelper.cpp b/Engine/Source/Programs/Online/OSSTestsCore/OnlineSubsystemCatchHelper.cpp
--- a/Engine/Source/Programs/Online/OSSTestsCore/OnlineSubsystemCatchHelper.cpp
+++ b/Engine/Source/Programs/Online/OSSTestsCore/OnlineSubsystemCatchHelper.cpp
@@ -121,6 +121,16 @@ void OnlineSubsystemTestBase::UnloadServiceModules()
 	}
 }
 
+void OnlineSubsystemTestBase::DestroyCurrentOnlineSubsystemModule() const
+{
+	FName SubsystemName = FName(GetSubsystem());
+	FModuleManager& ModuleManager = FModuleManager::Get();
+	FName ModuleName = SubsystemName;
+	const bool bIsShutdown = false;
+	ModuleManager.UnloadModule(ModuleName, bIsShutdown);
+	FModuleManager::LoadModulePtr<IModuleInterface>(ModuleName);
+}
+
 void OnlineSubsystemTestBase::ConstructInternal(FString SubsystemName)
 {
 	Subsystem = SubsystemName;
diff --git a/Engine/Source/Programs/Online/OSSTestsCore/OnlineSubsystemCatchHelper.h b/Engine/Source/Programs/Online/OSSTestsCore/OnlineSubsystemCatchHelper.h
--- a/Engine/Source/Programs/Online/OSSTestsCore/OnlineSubsystemCatchHelper.h
+++ b/Engine/Source/Programs/Online/OSSTestsCore/OnlineSubsystemCatchHelper.h
@@ -31,6 +31,8 @@ protected:
 	
 	FString GetSubsystem() const;
 
+	void DestroyCurrentOnlineSubsystemModule() const;
+
 #if OSSTESTS_USEEXTERNAUTH
 	TArray<FOnlineAccountCredentials> CustomCredentials(int32 LocalUserNum, int32 NumUsers) const;
 #endif // OSSTESTS_USEEXTERNAUTH
diff --git a/Engine/Source/Programs/Online/OSSTestsCore/Tests/Identity/IdentityAutoLoginTests.cpp b/Engine/Source/Programs/Online/OSSTestsCore/Tests/Identity/IdentityAutoLoginTests.cpp
--- a/Engine/Source/Programs/Online/OSSTestsCore/Tests/Identity/IdentityAutoLoginTests.cpp
+++ b/Engine/Source/Programs/Online/OSSTestsCore/Tests/Identity/IdentityAutoLoginTests.cpp
@@ -9,9 +9,41 @@
 
 #define IDENTITY_TAG "[suite_identity]"
 #define EG_IDENTITY_AUTOLOGIN_TAG IDENTITY_TAG "[autologin]"
+#define EG_IDENTITY_AUTOLOGIN_LEGACY_FLOW_TAG IDENTITY_TAG "[autologin][.NULL]"
 
 #define IDENTITY_TEST_CASE(x, ...) ONLINESUBSYSTEM_TEST_CASE(x, IDENTITY_TAG __VA_ARGS__)
 
+IDENTITY_TEST_CASE("Verify calling Identity AutoLogin with valid inputs and using legacy login flow returns the expected result(Success Case)", EG_IDENTITY_AUTOLOGIN_LEGACY_FLOW_TAG)
+{
+	GConfig->SetString(TEXT("/Script/OnlineSubsystemEOS.EOSSettings"), TEXT("bUseNewLoginFlow"), TEXT("false"), GEngineIni);
+	DestroyCurrentOnlineSubsystemModule();
+
+	FTestDriver LocalDriver;
+	int32 LocalUserNum = 0;
+	int32 NumUsers = 1;
+
+	TArray<FOnlineAccountCredentials> AccountCreds = GetCredentials(LocalUserNum, NumUsers);
+
+	FString LoginCredentialsType = TEXT("AUTH_TYPE=") + AccountCreds[0].Type + ",";
+	FString LoginCredentialsId = TEXT("AUTH_LOGIN=") + AccountCreds[0].Id + +",";
+	FString LoginCredentialsPassword = TEXT("AUTH_PASSWORD=") + AccountCreds[0].Token;
+
+	FCommandLine::Set(*(LoginCredentialsType + LoginCredentialsId + LoginCredentialsPassword));
+
+	FTestPipeline LocalPipeline = LocalDriver.MakePipeline()
+		.EmplaceStep<FIdentityAutoLoginStep>(LocalUserNum)
+		.EmplaceStep<FIdentityLogoutStep>(LocalUserNum);
+
+	FPipelineTestContext TestContext = FPipelineTestContext(FName(GetSubsystem()));
+	REQUIRE(LocalDriver.AddPipeline(MoveTemp(LocalPipeline), TestContext));
+	LocalDriver.RunToCompletion();
+
+	bool bUseNewLoginFlow = false;
+	GConfig->SetString(TEXT("/Script/OnlineSubsystemEOS.EOSSettings"), TEXT("bUseNewLoginFlow"), TEXT("true"), GEngineIni);
+	GConfig->GetBool(TEXT("/Script/OnlineSubsystemEOS.EOSSettings"), TEXT("bUseNewLoginFlow"), bUseNewLoginFlow, GEngineIni);
+	DestroyCurrentOnlineSubsystemModule();
+}
+
 IDENTITY_TEST_CASE("Verify calling Identity AutoLogin with valid inputs returns the expected result(Success Case)", EG_IDENTITY_AUTOLOGIN_TAG)
 { 
 	FTestDriver LocalDriver;
@@ -30,6 +62,7 @@ IDENTITY_TEST_CASE("Verify calling Identity AutoLogin with valid inputs returns
 		.EmplaceStep<FIdentityAutoLoginStep>(LocalUserNum)
 		.EmplaceStep<FIdentityLogoutStep>(LocalUserNum);
 
-	REQUIRE(LocalDriver.AddPipeline(MoveTemp(LocalPipeline), FName(GetSubsystem())));
+	FPipelineTestContext TestContext = FPipelineTestContext(FName(GetSubsystem()));
+	REQUIRE(LocalDriver.AddPipeline(MoveTemp(LocalPipeline), TestContext));
 	LocalDriver.RunToCompletion();
 }
diff --git a/Engine/Source/Programs/Online/OSSTestsCore/Tests/Identity/IdentityLoginTests.cpp b/Engine/Source/Programs/Online/OSSTestsCore/Tests/Identity/IdentityLoginTests.cpp
--- a/Engine/Source/Programs/Online/OSSTestsCore/Tests/Identity/IdentityLoginTests.cpp
+++ b/Engine/Source/Programs/Online/OSSTestsCore/Tests/Identity/IdentityLoginTests.cpp
@@ -6,9 +6,25 @@
 
 #define IDENTITY_TAG "[suite_identity]"
 #define EG_IDENTITY_LOGIN_TAG IDENTITY_TAG "[login]"
+#define EG_IDENTITY_LOGIN_LEGACY_FLOW_TAG IDENTITY_TAG "[login][.NULL]"
 
 #define IDENTITY_TEST_CASE(x, ...) ONLINESUBSYSTEM_TEST_CASE(x, IDENTITY_TAG __VA_ARGS__)
 
+IDENTITY_TEST_CASE("Verify calling Identity Login with valid inputs and using legacy login flow returns the expected result(Success Case)", EG_IDENTITY_LOGIN_LEGACY_FLOW_TAG)
+{
+	GConfig->SetString(TEXT("/Script/OnlineSubsystemEOS.EOSSettings"), TEXT("bUseNewLoginFlow"), TEXT("false"), GEngineIni);
+	DestroyCurrentOnlineSubsystemModule();
+
+	int32 NumUsersToImplicitLogin = 2;
+
+	GetLoginPipeline(NumUsersToImplicitLogin);
+
+	RunToCompletion();
+
+	GConfig->SetString(TEXT("/Script/OnlineSubsystemEOS.EOSSettings"), TEXT("bUseNewLoginFlow"), TEXT("true"), GEngineIni);
+	DestroyCurrentOnlineSubsystemModule();
+}
+
 IDENTITY_TEST_CASE("Verify calling Identity Login with valid inputs returns the expected result(Success Case)", EG_IDENTITY_LOGIN_TAG)
 {
 	int32 NumUsersToImplicitLogin = 2;
