Rework ONLINE_TEST_CASE to provide a fixture that tests can override

- tests can be run only for a certain type of service
- added possibility of using fixtures for any tests

diff --git a/Engine/Source/Programs/Online/OSSTestsCore/Helpers/AutoRestoreConfig.h b/Engine/Source/Programs/Online/OSSTestsCore/Helpers/AutoRestoreConfig.h
new file mode 100644
--- /dev/null
+++ b/Engine/Source/Programs/Online/OSSTestsCore/Helpers/AutoRestoreConfig.h
@@ -0,0 +1,75 @@
+// Copyright Epic Games, Inc. All Rights Reserved.
+
+#pragma once 
+
+template<typename T>
+class TAutoRestoreGConfig
+{
+public:
+	TAutoRestoreGConfig(const TCHAR* InSectionName, const TCHAR* InKeyName, const FString& InFileName)
+		: SectionName(InSectionName), KeyName(InKeyName), FileName(InFileName)
+	{
+	}
+
+	~TAutoRestoreGConfig()
+	{
+		Reset();
+	}
+
+	void Reset()
+	{
+		if (bRetrievedInitialValue)
+		{
+			if (InitialValue.IsSet())
+			{
+				SetValue(MoveTempIfPossible(InitialValue.GetValue()));
+				InitialValue.Reset();
+			}
+			else
+			{
+				FConfigFile* ConfigFile = GConfig->FindConfigFile(FileName);
+				CHECK(ConfigFile);
+				ConfigFile->RemoveKeyFromSection(SectionName, KeyName);
+			}
+			bRetrievedInitialValue = false;
+		}
+	}
+
+	bool IsSet()
+	{
+		return bRetrievedInitialValue;
+	}
+
+	void SetValue(const T& NewValue)
+	{
+		if (!bRetrievedInitialValue)
+		{
+			T PrevValue = {};
+			if (GConfig->GetValue(SectionName, KeyName, PrevValue, FileName))
+			{
+				InitialValue.Emplace(PrevValue);
+			}
+			bRetrievedInitialValue = true;
+		}
+
+		if constexpr (std::is_same_v<T, FString>)
+		{
+			GConfig->SetString(SectionName, KeyName, *NewValue, FileName);
+		}
+		else if constexpr (std::is_same_v<T, bool>)
+		{
+			GConfig->SetBool(SectionName, KeyName, NewValue, FileName);
+		}
+		else
+		{
+			checkNoEntry();
+		}
+	}
+
+private:
+	TOptional<T> InitialValue;
+	bool bRetrievedInitialValue = false;
+	const TCHAR* SectionName;
+	const TCHAR* KeyName;
+	const FString& FileName;
+};
diff --git a/Engine/Source/Programs/Online/OSSTestsCore/Helpers/Friends/FriendsEnsureFriendshipHelper.h b/Engine/Source/Programs/Online/OSSTestsCore/Helpers/Friends/FriendsEnsureFriendshipHelper.h
--- a/Engine/Source/Programs/Online/OSSTestsCore/Helpers/Friends/FriendsEnsureFriendshipHelper.h
+++ b/Engine/Source/Programs/Online/OSSTestsCore/Helpers/Friends/FriendsEnsureFriendshipHelper.h
@@ -15,13 +15,13 @@
 #include "Helpers/Friends/FriendsIsFriendHelper.h"
 #include "Helpers/Friends/FriendsUnblockPlayerHelper.h"
 
-#include "Helpers/CommonHelper.h"
+#include "Helpers/InnerStepBuilder.h"
 
 struct FFriendsEnsureFriendshipStep : public FTestPipeline::FStep
 {
 	FFriendsEnsureFriendshipStep(int32 InLocalUserNum, int32 InTargetUserNum, FUniqueNetIdPtr* InLocalUserId, FUniqueNetIdPtr* InTargetUserId, FString& InLocalListName, bool InbIsFriendsListPopulated = false)
 		: InnerSteps(
-			FInnerStepArrayBuilder<FTestPipeline::FStep>()
+			TInnerStepArrayBuilder<FTestPipeline::FStep>()
 			.EmplaceInnerStep<FFriendsDeleteFriendListStep>(InLocalUserNum, InLocalListName)
 			.EmplaceInnerStep<FFriendsDeleteFriendListStep>(InTargetUserNum, InLocalListName)
 			.EmplaceInnerStep<FFriendsUnblockPlayerStep>(InLocalUserNum, InTargetUserId)
diff --git a/Engine/Source/Programs/Online/OSSTestsCore/Helpers/InnerStepBuilder.h b/Engine/Source/Programs/Online/OSSTestsCore/Helpers/InnerStepBuilder.h
new file mode 100644
--- /dev/null
+++ b/Engine/Source/Programs/Online/OSSTestsCore/Helpers/InnerStepBuilder.h
@@ -0,0 +1,19 @@
+// Copyright Epic Games, Inc. All Rights Reserved.
+
+#pragma once 
+
+template<typename TBaseInnerStep>
+struct TInnerStepArrayBuilder
+{
+	TInnerStepArrayBuilder()
+	{}
+
+	template<typename T, typename... TArguments>
+	TInnerStepArrayBuilder&& EmplaceInnerStep(TArguments&&... Args)
+	{
+		Steps.Emplace(new T(std::forward<TArguments>(Args)...));	
+		return std::move(*this);
+	}
+
+	TArray<TUniquePtr<TBaseInnerStep>> Steps;
+};
diff --git a/Engine/Source/Programs/Online/OSSTestsCore/Helpers/OnlineSubsystemTestFixtures.h b/Engine/Source/Programs/Online/OSSTestsCore/Helpers/OnlineSubsystemTestFixtures.h
new file mode 100644
--- /dev/null
+++ b/Engine/Source/Programs/Online/OSSTestsCore/Helpers/OnlineSubsystemTestFixtures.h
@@ -0,0 +1,31 @@
+// Copyright Epic Games, Inc. All Rights Reserved.
+
+#pragma once 
+
+#include "OnlineSubsystemCatchHelper.h"
+#include "Helpers/AutoRestoreConfig.h"
+
+class FOnlineSubsystemEOSLegacyTestFixture : public FOnlineSubsystemTestBaseFixture
+{ 
+public:
+	FOnlineSubsystemEOSLegacyTestFixture()
+		: ConfigInstanceName(TEXT("/Script/OnlineSubsystemEOS.EOSSettings"), TEXT("bUseNewLoginFlow"), GEngineIni)
+	{
+		bool bUseNewLoginFlow = false;
+		ConfigInstanceName.SetValue(bUseNewLoginFlow);
+	}
+
+	~FOnlineSubsystemEOSLegacyTestFixture()
+	{
+		DestroyCurrentOnlineSubsystemModule();
+	}
+
+	void DestroyCurrentOnlineSubsystemModule() const override 
+	{
+		ConfigInstanceName.Reset();
+		FOnlineSubsystemTestBaseFixture::DestroyCurrentOnlineSubsystemModule();
+	}
+
+private:	
+	mutable TAutoRestoreGConfig<bool> ConfigInstanceName;
+};
diff --git a/Engine/Source/Programs/Online/OSSTestsCore/OnlineSubsystemCatchHelper.cpp b/Engine/Source/Programs/Online/OSSTestsCore/OnlineSubsystemCatchHelper.cpp
--- a/Engine/Source/Programs/Online/OSSTestsCore/OnlineSubsystemCatchHelper.cpp
+++ b/Engine/Source/Programs/Online/OSSTestsCore/OnlineSubsystemCatchHelper.cpp
@@ -41,7 +41,7 @@ TArray<TFunction<void()>>* GetGlobalInitalizers()
 	return &gInitalizersToCallInMain;
 }
 
-TArray<OnlineSubsystemAutoReg::FApplicableServicesConfig> OnlineSubsystemAutoReg::GetApplicableServices()
+TArray<FOnlineSubsystemAutoReg::FApplicableServicesConfig> FOnlineSubsystemAutoReg::GetApplicableServices()
 {
 	static TArray<FApplicableServicesConfig> ServicesConfig =
 		[]()
@@ -89,7 +89,7 @@ TArray<FString> GetServiceModules()
 {
 	TArray<FString> Modules;
 
-	for (const OnlineSubsystemAutoReg::FApplicableServicesConfig& Config : OnlineSubsystemAutoReg::GetApplicableServices())
+	for (const FOnlineSubsystemAutoReg::FApplicableServicesConfig& Config : FOnlineSubsystemAutoReg::GetApplicableServices())
 	{
 		for (const FString& Module : Config.ModulesToLoad)
 		{
@@ -100,7 +100,7 @@ TArray<FString> GetServiceModules()
 	return Modules;
 }
 
-void OnlineSubsystemTestBase::LoadServiceModules()
+void FOnlineSubsystemTestBaseFixture::LoadServiceModules()
 {
 	for (const FString& Module : GetServiceModules())
 	{
@@ -108,7 +108,7 @@ void OnlineSubsystemTestBase::LoadServiceModules()
 	}
 }
 
-void OnlineSubsystemTestBase::UnloadServiceModules()
+void FOnlineSubsystemTestBaseFixture::UnloadServiceModules()
 {
 	const TArray<FString>& Modules = GetServiceModules();
 	// Shutdown in reverse order
@@ -121,7 +121,7 @@ void OnlineSubsystemTestBase::UnloadServiceModules()
 	}
 }
 
-void OnlineSubsystemTestBase::DestroyCurrentOnlineSubsystemModule() const
+void FOnlineSubsystemTestBaseFixture::DestroyCurrentOnlineSubsystemModule() const
 {
 	FName SubsystemName = FName(GetSubsystem());
 	FModuleManager& ModuleManager = FModuleManager::Get();
@@ -131,29 +131,28 @@ void OnlineSubsystemTestBase::DestroyCurrentOnlineSubsystemModule() const
 	FModuleManager::LoadModulePtr<IModuleInterface>(ModuleName);
 }
 
-void OnlineSubsystemTestBase::ConstructInternal(FString SubsystemName)
+void FOnlineSubsystemTestBaseFixture::ConstructInternal(FString SubsystemName)
 {
 	Subsystem = SubsystemName;
 }
 
-OnlineSubsystemTestBase::OnlineSubsystemTestBase()
+FOnlineSubsystemTestBaseFixture::FOnlineSubsystemTestBaseFixture()
 	: Driver()
 	, Pipeline(Driver.MakePipeline())
 {
 	// handle most cxn in ConstructInternal
 }
 
-OnlineSubsystemTestBase::~OnlineSubsystemTestBase()
+FOnlineSubsystemTestBaseFixture::~FOnlineSubsystemTestBaseFixture()
 {
-
 }
 
-FString OnlineSubsystemTestBase::GetSubsystem() const
+FString FOnlineSubsystemTestBaseFixture::GetSubsystem() const
 {
 	return Subsystem;
 }
 
-TArray<FOnlineAccountCredentials> OnlineSubsystemTestBase::GetIniCredentials(int32 NumUsers) const
+TArray<FOnlineAccountCredentials> FOnlineSubsystemTestBaseFixture::GetIniCredentials(int32 NumUsers) const
 {
 	FString LoginCredentialCategory = FString::Printf(TEXT("LoginCredentials %s"), *Subsystem);
 	TArray<FString> CredentialsArr;
@@ -179,7 +178,7 @@ TArray<FOnlineAccountCredentials> OnlineSubsystemTestBase::GetIniCredentials(int
 	return OnlineAccountCredentials;
 }
 
-TArray<FOnlineAccountCredentials> OnlineSubsystemTestBase::GetCredentials(int32 LocalUserNum, int32 NumUsers) const
+TArray<FOnlineAccountCredentials> FOnlineSubsystemTestBaseFixture::GetCredentials(int32 LocalUserNum, int32 NumUsers) const
 {
 #if OSSTESTS_USEEXTERNAUTH
 	return CustomCredentials(LocalUserNum, NumUsers);
@@ -188,12 +187,12 @@ TArray<FOnlineAccountCredentials> OnlineSubsystemTestBase::GetCredentials(int32
 #endif // OSSTESTS_USEEXTERNAUTH
 }
 
-FString OnlineSubsystemTestBase::GetLoginCredentialCategory() const
+FString FOnlineSubsystemTestBaseFixture::GetLoginCredentialCategory() const
 {
 	return FString::Printf(TEXT("LoginCredentials %s"), *Subsystem);
 }
 
-FTestPipeline& OnlineSubsystemTestBase::GetLoginPipeline(uint32 NumUsersToLogin) const
+FTestPipeline& FOnlineSubsystemTestBaseFixture::GetLoginPipeline(uint32 NumUsersToLogin) const
 {
 	REQUIRE(NumLocalUsers == -1); // Don't call GetLoginPipeline more than once per test
 	NumLocalUsers = NumUsersToLogin;
@@ -230,12 +229,12 @@ FTestPipeline& OnlineSubsystemTestBase::GetLoginPipeline(uint32 NumUsersToLogin)
 	return Pipeline;
 }
 
-FTestPipeline& OnlineSubsystemTestBase::GetPipeline() const
+FTestPipeline& FOnlineSubsystemTestBaseFixture::GetPipeline() const
 {
 	return GetLoginPipeline(0);
 }
 
-void OnlineSubsystemTestBase::RunToCompletion() const
+void FOnlineSubsystemTestBaseFixture::RunToCompletion() const
 {
 	bool bUseAutoLogin = false;
 	bool bUseImplicitLogin = false;
@@ -267,15 +266,14 @@ void OnlineSubsystemTestBase::RunToCompletion() const
 	Driver.RunToCompletion();
 }
 
-// TODO: This should poll some info from the tags and generate the list based on that
-TArray<FString> OnlineSubsystemAutoReg::GetApplicableSubsystems()
+TArray<FString> FOnlineSubsystemAutoReg::GetApplicableSubsystems()
 {
 	TArray<FString> Subsystems;
 	GConfig->GetArray(TEXT("OnlineSubsystemTests"), TEXT("Subsystems"), Subsystems, GEngineIni);
 	return Subsystems;
 }
 
-bool OnlineSubsystemAutoReg::CheckAllTagsIsIn(const TArray<FString>& TestTags, const TArray<FString>& InputTags)
+bool FOnlineSubsystemAutoReg::CheckAllTagsIsIn(const TArray<FString>& TestTags, const TArray<FString>& InputTags)
 {
 	if (InputTags.Num() == 0)
 	{
@@ -305,7 +303,7 @@ bool OnlineSubsystemAutoReg::CheckAllTagsIsIn(const TArray<FString>& TestTags, c
 	return bAllInputTagsInTestTags;
 }
 
-bool OnlineSubsystemAutoReg::CheckAllTagsIsIn(const TArray<FString>& TestTags, const FString& RawTagString)
+bool FOnlineSubsystemAutoReg::CheckAllTagsIsIn(const TArray<FString>& TestTags, const FString& RawTagString)
 {
 	TArray<FString> InputTags;
 	RawTagString.ParseIntoArray(InputTags, TEXT(","));
@@ -318,7 +316,7 @@ bool OnlineSubsystemAutoReg::CheckAllTagsIsIn(const TArray<FString>& TestTags, c
 	return CheckAllTagsIsIn(TestTags, InputTags);
 }
 
-FString OnlineSubsystemAutoReg::GenerateTags(const FString& ServiceName, const FReportingSkippableTags& SkippableTags, const TCHAR* InTag)
+FString FOnlineSubsystemAutoReg::GenerateTags(const FString& ServiceName, const FReportingSkippableTags& SkippableTags, const TCHAR* InTag)
 {
 	//Copy String here for ease-of-manipulation
 	FString RawInTag = InTag;
@@ -359,7 +357,7 @@ FString OnlineSubsystemAutoReg::GenerateTags(const FString& ServiceName, const F
 	return FString::Printf(TEXT("[%s] %s"), *ServiceName, *RawInTag);
 }
 
-bool OnlineSubsystemAutoReg::ShouldDisableTest(const FString& ServiceName, const FReportingSkippableTags& SkippableTags, const TCHAR* InTag)
+bool FOnlineSubsystemAutoReg::ShouldDisableTest(const FString& ServiceName, const FReportingSkippableTags& SkippableTags, const TCHAR* InTag)
 {
 	//Copy String here for ease-of-manipulation
 	const FString RawInTag = InTag;
@@ -380,6 +378,23 @@ bool OnlineSubsystemAutoReg::ShouldDisableTest(const FString& ServiceName, const
 		return true;
 	}
 
+	// Check for exclusive runs
+	for (const FApplicableServicesConfig& Config : GetApplicableServices())
+	{
+		const FString& ServiceTag = Config.Tag;
+		if (ServiceName.Equals(ServiceTag, ESearchCase::IgnoreCase))
+		{
+			continue;
+		}
+
+		// If we contain [.NULL] and we're running with [EOS] we shouldn't
+		// generate a test for [EOS] here.
+		if (RawInTag.Contains("." + ServiceTag))
+		{
+			return true;
+		}
+	}
+
 	// If we contain tags from config it means 
 	// we shouldn't run this test
 	for (const FString& DisableTag : SkippableTags.DisableTestTags)
@@ -395,11 +410,13 @@ bool OnlineSubsystemAutoReg::ShouldDisableTest(const FString& ServiceName, const
 }
 
 // This code is kept identical to Catch internals so that there is as little deviation from OSS_TESTS and Online_OSS_TESTS as possible
-OnlineSubsystemAutoReg::OnlineSubsystemAutoReg(OnlineSubsystemTestConstructor TestCtor, Catch::SourceLineInfo LineInfo, const char* Name, const char* Tags, const char* AddlOnlineInfo)
+FOnlineSubsystemAutoReg::FOnlineSubsystemAutoReg(Catch::Detail::unique_ptr<FOnlineSubsystemFixtureInvoker> TestInvoker, Catch::SourceLineInfo LineInfo, const char* Name, const char* Tags, const char* AddlOnlineInfo)
 {
+	auto TestInvokerPtr = TestInvoker.release();
+
 	auto GlobalInitalizersPtr = GetGlobalInitalizers();
 	ensure(GlobalInitalizersPtr);
-	GlobalInitalizersPtr->Add([=, this]() -> void
+	GlobalInitalizersPtr->Add([TestInvokerPtr = MoveTemp(TestInvokerPtr), LineInfo, Name, Tags, this]() mutable -> void
 		{
 			for (const FString& Subsystem : GetApplicableSubsystems())
 			{
@@ -418,17 +435,16 @@ OnlineSubsystemAutoReg::OnlineSubsystemAutoReg(OnlineSubsystemTestConstructor Te
 					continue;
 				}
 
-				// TestCtor will create a new instance of the test we are calling- ConstructInternal is separate so that we can pass any arguments we want instead of baking them into the macro
-				OnlineSubsystemTestBase* NewTest = TestCtor();
-				NewTest->ConstructInternal(Subsystem);
+				TestInvokerPtr->SetSubsystem(Subsystem);
 
 				// This code is lifted from Catch internals to register a test
 				Catch::getMutableRegistryHub().registerTest(Catch::makeTestCaseInfo(
 					std::string(Catch::StringRef()),  // Used for testing a static method instead of a function- not needed since we're passing an ITestInvoker macro
 					Catch::NameAndTags{ NewName.Get(), NewTags.Get() },
 					LineInfo),
-					Catch::Detail::unique_ptr(NewTest) // This is taking the ITestInvoker macro and will call invoke() to run the test
+					Catch::Detail::unique_ptr(TestInvokerPtr) // This is taking the ITestInvoker macro and will call invoke() to run the test
 				);
 			}
 		});
-}
\ No newline at end of file
+}
+
diff --git a/Engine/Source/Programs/Online/OSSTestsCore/OnlineSubsystemCatchHelper.h b/Engine/Source/Programs/Online/OSSTestsCore/OnlineSubsystemCatchHelper.h
--- a/Engine/Source/Programs/Online/OSSTestsCore/OnlineSubsystemCatchHelper.h
+++ b/Engine/Source/Programs/Online/OSSTestsCore/OnlineSubsystemCatchHelper.h
@@ -13,11 +13,57 @@
 #include <catch2/catch_test_case_info.hpp>
 #include <catch2/interfaces/catch_interfaces_registry_hub.hpp>
 
-class OnlineSubsystemTestBase : public Catch::ITestInvoker
+class FOnlineSubsystemFixtureInvoker : public Catch::ITestInvoker
+{
+public:
+	virtual void invoke() const = 0;
+
+	void SetSubsystem(const FString& InStoredSubsystem) {
+		StoredSubsystem = InStoredSubsystem;
+	}
+
+protected:
+	FString StoredSubsystem;
+};
+
+template <typename C>
+class TestInvokerFixture : public FOnlineSubsystemFixtureInvoker
+{
+public:
+	constexpr TestInvokerFixture(void (C::* InTestAsMethod)() const) noexcept :
+		TestAsMethod(InTestAsMethod) {}
+
+	void PrepareTestCase() const {
+		Fixture = Catch::Detail::make_unique<C>();
+		Fixture->ConstructInternal(StoredSubsystem);
+	}
+
+	void TearDownTestCase() const {
+		Fixture.reset();
+	}
+
+	void invoke() const override {
+		PrepareTestCase();
+		auto* f = Fixture.get();
+		(f->*TestAsMethod)();
+		TearDownTestCase();
+	}
+
+private:
+	void (C::* TestAsMethod)() const;
+	mutable Catch::Detail::unique_ptr<C> Fixture = nullptr;
+};
+
+template<typename C>
+Catch::Detail::unique_ptr<FOnlineSubsystemFixtureInvoker> MakeTestInvokerFixture(void (C::* TestAsMethod)() const) {
+	return Catch::Detail::make_unique<TestInvokerFixture<C>>(TestAsMethod);
+}
+
+class FOnlineSubsystemTestBaseFixture : public Catch::ITestInvoker
 {
 public:
 	void ConstructInternal(FString ServiceName);
-	virtual ~OnlineSubsystemTestBase();
+	virtual ~FOnlineSubsystemTestBaseFixture();
 
 	/* Loads all necessary services for the current test run */
 	static void LoadServiceModules();
@@ -27,11 +73,11 @@ public:
 
 protected:
 
-	OnlineSubsystemTestBase();
+	FOnlineSubsystemTestBaseFixture();
 	
 	FString GetSubsystem() const;
 
-	void DestroyCurrentOnlineSubsystemModule() const;
+	virtual void DestroyCurrentOnlineSubsystemModule() const;
 
 #if OSSTESTS_USEEXTERNAUTH
 	TArray<FOnlineAccountCredentials> CustomCredentials(int32 LocalUserNum, int32 NumUsers) const;
@@ -62,9 +108,7 @@ private:
 	mutable uint32 NumUsersToLogout = -1;
 };
 
-typedef OnlineSubsystemTestBase* (*OnlineSubsystemTestConstructor)();
-
-class OnlineSubsystemAutoReg
+class FOnlineSubsystemAutoReg
 {
 private:
 	// TODO: This should poll some info from the tags and generate the list based on that
@@ -112,30 +156,30 @@ public:
 	static bool ShouldDisableTest(const FString& ServiceName, const FReportingSkippableTags& SkippableTags, const TCHAR* InTag);
 
 	// This code is kept identical to Catch internals so that there is as little deviation from OSS_TESTS and Online_OSS_TESTS as possible
-	OnlineSubsystemAutoReg(OnlineSubsystemTestConstructor TestCtor, Catch::SourceLineInfo LineInfo, const char* Name, const char* Tags, const char* AddlOnlineInfo);
+	FOnlineSubsystemAutoReg(Catch::Detail::unique_ptr<FOnlineSubsystemFixtureInvoker> TestInvoker, Catch::SourceLineInfo LineInfo, const char* Name, const char* Tags, const char* AddlOnlineInfo);
+
 };
 
 TArray<TFunction<void()>>* GetGlobalInitalizers();
 
-
-// INTERNAL_CATCH_UNIQUE_NAME is what Catch uses to translate the name of the test into a unique identifier without any spaces
-
-#define INTERNAL_ONLINESUBSYSTEM_TEST_CASE_NAMED(RegName, Name, Tags, ...)\
+#define INTERNAL_ONLINESUBSYSTEM_TEST_CASE_NAMED_FIXTURE(RegName, ClassName, Name, Tags, ...)\
 namespace {\
-class PREPROCESSOR_JOIN(OnlineSubsystemTest_,RegName) : public OnlineSubsystemTestBase {\
-protected:\
+class PREPROCESSOR_JOIN(OnlineSubsystemTest_,RegName) : public ClassName {\
+public:\
 	virtual void invoke() const override;\
+	void test() const;\
 };\
-	OnlineSubsystemTestBase* PREPROCESSOR_JOIN(Construct_,RegName) (){\
-		return new PREPROCESSOR_JOIN(OnlineSubsystemTest_,RegName) ();\
-	}\
-	OnlineSubsystemAutoReg RegName( PREPROCESSOR_JOIN(&Construct_,RegName) , CATCH_INTERNAL_LINEINFO , Name, Tags, "");\
+	FOnlineSubsystemAutoReg RegName( ::MakeTestInvokerFixture( PREPROCESSOR_JOIN(&OnlineSubsystemTest_,RegName)::test ), CATCH_INTERNAL_LINEINFO, Name, Tags, "");\
 }\
+void PREPROCESSOR_JOIN(OnlineSubsystemTest_,RegName)::test() const { invoke(); }\
 void PREPROCESSOR_JOIN(OnlineSubsystemTest_,RegName)::invoke() const\
 
+#define ONLINESUBSYSTEM_TEST_CASE_FIXTURE(ClassName, Name, Tags, ...) \
+	INTERNAL_ONLINESUBSYSTEM_TEST_CASE_NAMED_FIXTURE(INTERNAL_CATCH_UNIQUE_NAME(OnlineSubsystemRegistrar), ClassName, Name, Tags, __VA_ARGS__)
+
 #define ONLINESUBSYSTEM_TEST_CASE(Name, Tags, ...) \
-	INTERNAL_ONLINESUBSYSTEM_TEST_CASE_NAMED(INTERNAL_CATCH_UNIQUE_NAME(OnlineSubsystemRegistrar), Name, Tags, __VA_ARGS__)
+	INTERNAL_ONLINESUBSYSTEM_TEST_CASE_NAMED_FIXTURE(INTERNAL_CATCH_UNIQUE_NAME(OnlineSubsystemRegistrar), FOnlineSubsystemTestBaseFixture, Name, Tags, __VA_ARGS__)
 
 #define REQUIRE_OP(Op)\
 	CAPTURE(Op);\
-	REQUIRE(Op.WasSuccessful());
\ No newline at end of file
+	REQUIRE(Op.WasSuccessful());
diff --git a/Engine/Source/Programs/Online/OSSTestsCore/Tests/Identity/IdentityAutoLoginTests.cpp b/Engine/Source/Programs/Online/OSSTestsCore/Tests/Identity/IdentityAutoLoginTests.cpp
--- a/Engine/Source/Programs/Online/OSSTestsCore/Tests/Identity/IdentityAutoLoginTests.cpp
+++ b/Engine/Source/Programs/Online/OSSTestsCore/Tests/Identity/IdentityAutoLoginTests.cpp
@@ -4,20 +4,21 @@
 
 #include "Helpers/Identity/IdentityAutoLoginHelper.h"
 #include "Helpers/Identity/IdentityLogoutHelper.h"
+#include "Helpers/OnlineSubsystemTestFixtures.h"
+
 #include "Misc/CommandLine.h"
 #include "OnlineSubsystemCatchHelper.h"
 
 #define IDENTITY_TAG "[suite_identity]"
 #define EG_IDENTITY_AUTOLOGIN_TAG IDENTITY_TAG "[autologin]"
-#define EG_IDENTITY_AUTOLOGIN_LEGACY_FLOW_TAG IDENTITY_TAG "[autologin][.NULL]"
+#define EG_IDENTITY_AUTOLOGIN_LEGACY_EOS_FLOW_TAG IDENTITY_TAG "[autologin][.EOS]"
 
 #define IDENTITY_TEST_CASE(x, ...) ONLINESUBSYSTEM_TEST_CASE(x, IDENTITY_TAG __VA_ARGS__)
 
-IDENTITY_TEST_CASE("Verify calling Identity AutoLogin with valid inputs and using legacy login flow returns the expected result(Success Case)", EG_IDENTITY_AUTOLOGIN_LEGACY_FLOW_TAG)
-{
-	GConfig->SetString(TEXT("/Script/OnlineSubsystemEOS.EOSSettings"), TEXT("bUseNewLoginFlow"), TEXT("false"), GEngineIni);
-	DestroyCurrentOnlineSubsystemModule();
+#define IDENTITY_LEGACY_EOS_TEST(x, ...) ONLINESUBSYSTEM_TEST_CASE_FIXTURE(FOnlineSubsystemEOSLegacyTestFixture, x, __VA_ARGS__)
 
+IDENTITY_LEGACY_EOS_TEST("Verify calling Identity AutoLogin with valid inputs and using legacy login flow returns the expected result(Success Case)", EG_IDENTITY_AUTOLOGIN_LEGACY_EOS_FLOW_TAG)
+{
 	FTestDriver LocalDriver;
 	int32 LocalUserNum = 0;
 	int32 NumUsers = 1;
@@ -37,11 +38,6 @@ IDENTITY_TEST_CASE("Verify calling Identity AutoLogin with valid inputs and usin
 	FPipelineTestContext TestContext = FPipelineTestContext(FName(GetSubsystem()));
 	REQUIRE(LocalDriver.AddPipeline(MoveTemp(LocalPipeline), TestContext));
 	LocalDriver.RunToCompletion();
-
-	bool bUseNewLoginFlow = false;
-	GConfig->SetString(TEXT("/Script/OnlineSubsystemEOS.EOSSettings"), TEXT("bUseNewLoginFlow"), TEXT("true"), GEngineIni);
-	GConfig->GetBool(TEXT("/Script/OnlineSubsystemEOS.EOSSettings"), TEXT("bUseNewLoginFlow"), bUseNewLoginFlow, GEngineIni);
-	DestroyCurrentOnlineSubsystemModule();
 }
 
 IDENTITY_TEST_CASE("Verify calling Identity AutoLogin with valid inputs returns the expected result(Success Case)", EG_IDENTITY_AUTOLOGIN_TAG)
diff --git a/Engine/Source/Programs/Online/OSSTestsCore/Tests/Identity/IdentityLoginTests.cpp b/Engine/Source/Programs/Online/OSSTestsCore/Tests/Identity/IdentityLoginTests.cpp
--- a/Engine/Source/Programs/Online/OSSTestsCore/Tests/Identity/IdentityLoginTests.cpp
+++ b/Engine/Source/Programs/Online/OSSTestsCore/Tests/Identity/IdentityLoginTests.cpp
@@ -3,26 +3,23 @@
 #include "TestDriver.h"
 
 #include "OnlineSubsystemCatchHelper.h"
+#include "Helpers/OnlineSubsystemTestFixtures.h"
 
 #define IDENTITY_TAG "[suite_identity]"
 #define EG_IDENTITY_LOGIN_TAG IDENTITY_TAG "[login]"
-#define EG_IDENTITY_LOGIN_LEGACY_FLOW_TAG IDENTITY_TAG "[login][.NULL]"
+#define EG_IDENTITY_LOGIN_LEGACY_FLOWEOS_TAG IDENTITY_TAG "[login][.EOS]"
 
 #define IDENTITY_TEST_CASE(x, ...) ONLINESUBSYSTEM_TEST_CASE(x, IDENTITY_TAG __VA_ARGS__)
 
-IDENTITY_TEST_CASE("Verify calling Identity Login with valid inputs and using legacy login flow returns the expected result(Success Case)", EG_IDENTITY_LOGIN_LEGACY_FLOW_TAG)
-{
-	GConfig->SetString(TEXT("/Script/OnlineSubsystemEOS.EOSSettings"), TEXT("bUseNewLoginFlow"), TEXT("false"), GEngineIni);
-	DestroyCurrentOnlineSubsystemModule();
+#define IDENTITY_LEGACY_EOS_TEST(x, ...) ONLINESUBSYSTEM_TEST_CASE_FIXTURE(FOnlineSubsystemEOSLegacyTestFixture, x, IDENTITY_TAG __VA_ARGS__)
 
+IDENTITY_LEGACY_EOS_TEST("Verify calling Identity Login with valid inputs and using legacy login flow returns the expected result(Success Case)", EG_IDENTITY_LOGIN_LEGACY_FLOWEOS_TAG)
+{
 	int32 NumUsersToImplicitLogin = 2;
 
 	GetLoginPipeline(NumUsersToImplicitLogin);
 
 	RunToCompletion();
-
-	GConfig->SetString(TEXT("/Script/OnlineSubsystemEOS.EOSSettings"), TEXT("bUseNewLoginFlow"), TEXT("true"), GEngineIni);
-	DestroyCurrentOnlineSubsystemModule();
 }
 
 IDENTITY_TEST_CASE("Verify calling Identity Login with valid inputs returns the expected result(Success Case)", EG_IDENTITY_LOGIN_TAG)
diff --git a/Engine/Source/Programs/Online/OSSTestsCore/Tests/OSSCatchHelpers/OSSCatchHelpers.cpp b/Engine/Source/Programs/Online/OSSTestsCore/Tests/OSSCatchHelpers/OSSCatchHelpers.cpp
--- a/Engine/Source/Programs/Online/OSSTestsCore/Tests/OSSCatchHelpers/OSSCatchHelpers.cpp
+++ b/Engine/Source/Programs/Online/OSSTestsCore/Tests/OSSCatchHelpers/OSSCatchHelpers.cpp
@@ -11,100 +11,100 @@
 SELFTEST_TEST_CASE("GenerateTags append MayFailTags case", GENERATETAGS_TAG)
 {
 	FString TestTags = "[foo][bar]";
-	OnlineSubsystemAutoReg::FReportingSkippableTags SkipTags;
+	FOnlineSubsystemAutoReg::FReportingSkippableTags SkipTags;
 	SkipTags.MayFailTags = { "[foo]" };
-	FString OutTags = OnlineSubsystemAutoReg::GenerateTags("TestService", SkipTags, *TestTags);
+	FString OutTags = FOnlineSubsystemAutoReg::GenerateTags("TestService", SkipTags, *TestTags);
 	CHECK(OutTags == TEXT("[TestService] [foo][bar][!mayfail]"));
 }
 
 SELFTEST_TEST_CASE("GenerateTags append multiple match MayFailTags case", GENERATETAGS_TAG)
 {
 	FString TestTags = "[foo][bar]";
-	OnlineSubsystemAutoReg::FReportingSkippableTags SkipTags;
+	FOnlineSubsystemAutoReg::FReportingSkippableTags SkipTags;
 	SkipTags.MayFailTags = { "[bar]", "[foo]" };
-	FString OutTags = OnlineSubsystemAutoReg::GenerateTags("TestService", SkipTags, *TestTags);
+	FString OutTags = FOnlineSubsystemAutoReg::GenerateTags("TestService", SkipTags, *TestTags);
 	CHECK(OutTags == TEXT("[TestService] [foo][bar][!mayfail]"));
 }
 
 SELFTEST_TEST_CASE("GenerateTags append by last match MayFailTags case", GENERATETAGS_TAG)
 {
 	FString TestTags = "[foo][bar]";
-	OnlineSubsystemAutoReg::FReportingSkippableTags SkipTags;
+	FOnlineSubsystemAutoReg::FReportingSkippableTags SkipTags;
 	SkipTags.MayFailTags = { "[wiz]", "[foo]" };
-	FString OutTags = OnlineSubsystemAutoReg::GenerateTags("TestService", SkipTags, *TestTags);
+	FString OutTags = FOnlineSubsystemAutoReg::GenerateTags("TestService", SkipTags, *TestTags);
 	CHECK(OutTags == TEXT("[TestService] [foo][bar][!mayfail]"));
 }
 
 SELFTEST_TEST_CASE("GenerateTags append by last match mutli-tag case", GENERATETAGS_TAG)
 {
 	FString TestTags = "[foo][bar]";
-	OnlineSubsystemAutoReg::FReportingSkippableTags SkipTags;
+	FOnlineSubsystemAutoReg::FReportingSkippableTags SkipTags;
 	SkipTags.MayFailTags = { "[foo],bar" };
-	FString OutTags = OnlineSubsystemAutoReg::GenerateTags("TestService", SkipTags, *TestTags);
+	FString OutTags = FOnlineSubsystemAutoReg::GenerateTags("TestService", SkipTags, *TestTags);
 	CHECK(OutTags == TEXT("[TestService] [foo][bar][!mayfail]"));
 }
 
 SELFTEST_TEST_CASE("GenerateTags append by last match mutli-tag no match case", GENERATETAGS_TAG)
 {
 	FString TestTags = "[foo][bar]";
-	OnlineSubsystemAutoReg::FReportingSkippableTags SkipTags;
+	FOnlineSubsystemAutoReg::FReportingSkippableTags SkipTags;
 	SkipTags.MayFailTags = { "[foo],[wiz]" };
-	FString OutTags = OnlineSubsystemAutoReg::GenerateTags("TestService", SkipTags, *TestTags);
+	FString OutTags = FOnlineSubsystemAutoReg::GenerateTags("TestService", SkipTags, *TestTags);
 	CHECK(OutTags == TEXT("[TestService] [foo][bar]"));
 }
 
 SELFTEST_TEST_CASE("GenerateTags don't append MayFailTags no match case", GENERATETAGS_TAG)
 {
-	OnlineSubsystemAutoReg::FReportingSkippableTags SkipTags;
+	FOnlineSubsystemAutoReg::FReportingSkippableTags SkipTags;
 	SkipTags.MayFailTags = { "[wiz]" };
 	FString TestTags = "[foo][bar]";
-	FString OutTags = OnlineSubsystemAutoReg::GenerateTags("TestService", SkipTags, *TestTags);
+	FString OutTags = FOnlineSubsystemAutoReg::GenerateTags("TestService", SkipTags, *TestTags);
 	CHECK(OutTags == TEXT("[TestService] [foo][bar]"));
 }
 
 SELFTEST_TEST_CASE("GenerateTags append ShouldFail case", GENERATETAGS_TAG)
 {
 	FString TestTags = "[foo][bar]";
-	OnlineSubsystemAutoReg::FReportingSkippableTags SkipTags;
+	FOnlineSubsystemAutoReg::FReportingSkippableTags SkipTags;
 	SkipTags.ShouldFailTags = { "[foo]" };
-	FString OutTags = OnlineSubsystemAutoReg::GenerateTags("TestService", SkipTags, *TestTags);
+	FString OutTags = FOnlineSubsystemAutoReg::GenerateTags("TestService", SkipTags, *TestTags);
 	CHECK(OutTags == TEXT("[TestService] [foo][bar][!shouldfail]"));
 }
 
 SELFTEST_TEST_CASE("GenerateTags append multiple match ShouldFail case", GENERATETAGS_TAG)
 {
 	FString TestTags = "[foo][bar]";
-	OnlineSubsystemAutoReg::FReportingSkippableTags SkipTags;
+	FOnlineSubsystemAutoReg::FReportingSkippableTags SkipTags;
 	SkipTags.ShouldFailTags = { "[bar]", "[foo]" };
-	FString OutTags = OnlineSubsystemAutoReg::GenerateTags("TestService", SkipTags, *TestTags);
+	FString OutTags = FOnlineSubsystemAutoReg::GenerateTags("TestService", SkipTags, *TestTags);
 	CHECK(OutTags == TEXT("[TestService] [foo][bar][!shouldfail]"));
 }
 
 SELFTEST_TEST_CASE("GenerateTags append by last match ShouldFail case", GENERATETAGS_TAG)
 {
 	FString TestTags = "[foo][bar]";
-	OnlineSubsystemAutoReg::FReportingSkippableTags SkipTags;
+	FOnlineSubsystemAutoReg::FReportingSkippableTags SkipTags;
 	SkipTags.ShouldFailTags = { "[wiz]", "[foo]" };
-	FString OutTags = OnlineSubsystemAutoReg::GenerateTags("TestService", SkipTags, *TestTags);
+	FString OutTags = FOnlineSubsystemAutoReg::GenerateTags("TestService", SkipTags, *TestTags);
 	CHECK(OutTags == TEXT("[TestService] [foo][bar][!shouldfail]"));
 }
 
 SELFTEST_TEST_CASE("GenerateTags don't append ShouldFail no match case", GENERATETAGS_TAG)
 {
 	FString TestTags = "[foo][bar]";
-	OnlineSubsystemAutoReg::FReportingSkippableTags SkipTags;
+	FOnlineSubsystemAutoReg::FReportingSkippableTags SkipTags;
 	SkipTags.ShouldFailTags = { "[wiz]" };
-	FString OutTags = OnlineSubsystemAutoReg::GenerateTags("TestService", SkipTags, *TestTags);
+	FString OutTags = FOnlineSubsystemAutoReg::GenerateTags("TestService", SkipTags, *TestTags);
 	CHECK(OutTags == TEXT("[TestService] [foo][bar]"));
 }
 
 SELFTEST_TEST_CASE("GenerateTags appends MayFail and ShouldFail Case", GENERATETAGS_TAG)
 {
-	OnlineSubsystemAutoReg::FReportingSkippableTags SkipTags;
+	FOnlineSubsystemAutoReg::FReportingSkippableTags SkipTags;
 	SkipTags.MayFailTags = { "[foo]" };
 	SkipTags.ShouldFailTags = { "[bar]" };
 	FString TestTags = "[foo][bar]";
-	FString OutTags = OnlineSubsystemAutoReg::GenerateTags("TestService", SkipTags, *TestTags);
+	FString OutTags = FOnlineSubsystemAutoReg::GenerateTags("TestService", SkipTags, *TestTags);
 	CHECK(OutTags == TEXT("[TestService] [foo][bar][!mayfail][!shouldfail]"));
 }
 
@@ -112,65 +112,65 @@ SELFTEST_TEST_CASE("GenerateTags appends MayFail and ShouldFail Case", GENERATET
 SELFTEST_TEST_CASE("GenerateTags append by last match mutli-tag ShouldFail case", GENERATETAGS_TAG)
 {
 	FString TestTags = "[foo][bar]";
-	OnlineSubsystemAutoReg::FReportingSkippableTags SkipTags;
+	FOnlineSubsystemAutoReg::FReportingSkippableTags SkipTags;
 	SkipTags.ShouldFailTags = { "[foo],bar" };
-	FString OutTags = OnlineSubsystemAutoReg::GenerateTags("TestService", SkipTags, *TestTags);
+	FString OutTags = FOnlineSubsystemAutoReg::GenerateTags("TestService", SkipTags, *TestTags);
 	CHECK(OutTags == TEXT("[TestService] [foo][bar][!shouldfail]"));
 }
 
 SELFTEST_TEST_CASE("GenerateTags append by last match mutli-tag no match ShouldFail case", GENERATETAGS_TAG)
 {
 	FString TestTags = "[foo][bar]";
-	OnlineSubsystemAutoReg::FReportingSkippableTags SkipTags;
+	FOnlineSubsystemAutoReg::FReportingSkippableTags SkipTags;
 	SkipTags.ShouldFailTags = { "[foo],[wiz]" };
-	FString OutTags = OnlineSubsystemAutoReg::GenerateTags("TestService", SkipTags, *TestTags);
+	FString OutTags = FOnlineSubsystemAutoReg::GenerateTags("TestService", SkipTags, *TestTags);
 	CHECK(OutTags == TEXT("[TestService] [foo][bar]"));
 }
 
 SELFTEST_TEST_CASE("GenerateTags appends no tags case", GENERATETAGS_TAG)
 {
-	OnlineSubsystemAutoReg::FReportingSkippableTags SkipTags;
+	FOnlineSubsystemAutoReg::FReportingSkippableTags SkipTags;
 	SkipTags.MayFailTags = { "[wiz]" };
 	SkipTags.ShouldFailTags = {  };
 	FString TestTags = "[foo][bar]";
-	FString OutTags = OnlineSubsystemAutoReg::GenerateTags("TestService", SkipTags, *TestTags);
+	FString OutTags = FOnlineSubsystemAutoReg::GenerateTags("TestService", SkipTags, *TestTags);
 	CHECK(OutTags == TEXT("[TestService] [foo][bar]"));
 }
 
 SELFTEST_TEST_CASE("ShouldDisableTest returns true on single-tag config", SHOULDDISABLETEST_TAG)
 {
-	OnlineSubsystemAutoReg::FReportingSkippableTags SkipTags;
+	FOnlineSubsystemAutoReg::FReportingSkippableTags SkipTags;
 	SkipTags.DisableTestTags = { "[foo]" };
 	FString TestTags = "[foo][bar]";
-	CHECK(OnlineSubsystemAutoReg::ShouldDisableTest("TestService", SkipTags, *TestTags) == true);
+	CHECK(FOnlineSubsystemAutoReg::ShouldDisableTest("TestService", SkipTags, *TestTags) == true);
 }
 
 SELFTEST_TEST_CASE("ShouldDisableTest returns true on multi-tag config", SHOULDDISABLETEST_TAG)
 {
-	OnlineSubsystemAutoReg::FReportingSkippableTags SkipTags;
+	FOnlineSubsystemAutoReg::FReportingSkippableTags SkipTags;
 	SkipTags.DisableTestTags = { "[foo],bar" };
 	FString TestTags = "[foo][bar]";
-	CHECK(OnlineSubsystemAutoReg::ShouldDisableTest("TestService", SkipTags, *TestTags) == true);
+	CHECK(FOnlineSubsystemAutoReg::ShouldDisableTest("TestService", SkipTags, *TestTags) == true);
 }
 
 SELFTEST_TEST_CASE("ShouldDisableTest returns true on !<service>", SHOULDDISABLETEST_TAG)
 {
 	FString TestTags = "[foo][bar][!TestService]";
-	CHECK(OnlineSubsystemAutoReg::ShouldDisableTest("TestService", {}, *TestTags) == true);
+	CHECK(FOnlineSubsystemAutoReg::ShouldDisableTest("TestService", {}, *TestTags) == true);
 }
 
 SELFTEST_TEST_CASE("ShouldDisableTest returns false with no tags and no config skips", SHOULDDISABLETEST_TAG)
 {
 	FString TestTags = "[foo][bar]";
-	CHECK(OnlineSubsystemAutoReg::ShouldDisableTest("TestService", {}, *TestTags) == false);
+	CHECK(FOnlineSubsystemAutoReg::ShouldDisableTest("TestService", {}, *TestTags) == false);
 }
 
 SELFTEST_TEST_CASE("ShouldDisableTest returns false with no matching no-tags and no matching multi-tag config skips", SHOULDDISABLETEST_TAG)
 {
-	OnlineSubsystemAutoReg::FReportingSkippableTags SkipTags;
+	FOnlineSubsystemAutoReg::FReportingSkippableTags SkipTags;
 	SkipTags.DisableTestTags = { "[foo],wiz" };
 	FString TestTags = "[foo][bar]";
-	CHECK(OnlineSubsystemAutoReg::ShouldDisableTest("TestService", SkipTags, *TestTags) == false);
+	CHECK(FOnlineSubsystemAutoReg::ShouldDisableTest("TestService", SkipTags, *TestTags) == false);
 }
 
 SELFTEST_TEST_CASE("CheckAllTagsIsIn(TArray, FString) true cases", CHECKALLTAGSISIN_TAG)
@@ -179,30 +179,30 @@ SELFTEST_TEST_CASE("CheckAllTagsIsIn(TArray, FString) true cases", CHECKALLTAGSI
 	CAPTURE(TestTags);
 
 	// Truthy Cases
-	CHECK(OnlineSubsystemAutoReg::CheckAllTagsIsIn(TestTags, "bob, alice") == true);
-	CHECK(OnlineSubsystemAutoReg::CheckAllTagsIsIn(TestTags, "bob,alice") == true);
-	CHECK(OnlineSubsystemAutoReg::CheckAllTagsIsIn(TestTags, " bob,alice ") == true);
-	CHECK(OnlineSubsystemAutoReg::CheckAllTagsIsIn(TestTags, "foo") == true);
-	CHECK(OnlineSubsystemAutoReg::CheckAllTagsIsIn(TestTags, ",foo") == true);
-	CHECK(OnlineSubsystemAutoReg::CheckAllTagsIsIn(TestTags, "bob,alice,foo") == true);
+	CHECK(FOnlineSubsystemAutoReg::CheckAllTagsIsIn(TestTags, "bob, alice") == true);
+	CHECK(FOnlineSubsystemAutoReg::CheckAllTagsIsIn(TestTags, "bob,alice") == true);
+	CHECK(FOnlineSubsystemAutoReg::CheckAllTagsIsIn(TestTags, " bob,alice ") == true);
+	CHECK(FOnlineSubsystemAutoReg::CheckAllTagsIsIn(TestTags, "foo") == true);
+	CHECK(FOnlineSubsystemAutoReg::CheckAllTagsIsIn(TestTags, ",foo") == true);
+	CHECK(FOnlineSubsystemAutoReg::CheckAllTagsIsIn(TestTags, "bob,alice,foo") == true);
 
 	//Bracket Parsing
-	CHECK(OnlineSubsystemAutoReg::CheckAllTagsIsIn(TestTags, "[bob],[alice],[foo]") == true);
-	CHECK(OnlineSubsystemAutoReg::CheckAllTagsIsIn(TestTags, "[bob], [alice,foo]") == true);
-	CHECK(OnlineSubsystemAutoReg::CheckAllTagsIsIn(TestTags, "bob],  alice],  [foo]  ,") == true);
-	CHECK(OnlineSubsystemAutoReg::CheckAllTagsIsIn(TestTags, ",[foo]") == true);
-	CHECK(OnlineSubsystemAutoReg::CheckAllTagsIsIn(TestTags, ",foo]") == true);
-	CHECK(OnlineSubsystemAutoReg::CheckAllTagsIsIn(TestTags, "[wiz]") == false);
+	CHECK(FOnlineSubsystemAutoReg::CheckAllTagsIsIn(TestTags, "[bob],[alice],[foo]") == true);
+	CHECK(FOnlineSubsystemAutoReg::CheckAllTagsIsIn(TestTags, "[bob], [alice,foo]") == true);
+	CHECK(FOnlineSubsystemAutoReg::CheckAllTagsIsIn(TestTags, "bob],  alice],  [foo]  ,") == true);
+	CHECK(FOnlineSubsystemAutoReg::CheckAllTagsIsIn(TestTags, ",[foo]") == true);
+	CHECK(FOnlineSubsystemAutoReg::CheckAllTagsIsIn(TestTags, ",foo]") == true);
+	CHECK(FOnlineSubsystemAutoReg::CheckAllTagsIsIn(TestTags, "[wiz]") == false);
 
 	// Negative Cases
-	CHECK(OnlineSubsystemAutoReg::CheckAllTagsIsIn(TestTags, "bob,alice,foo,wiz") == false);
-	CHECK(OnlineSubsystemAutoReg::CheckAllTagsIsIn(TestTags, "bob,wiz") == false);
-	CHECK(OnlineSubsystemAutoReg::CheckAllTagsIsIn(TestTags, ",wiz") == false);
+	CHECK(FOnlineSubsystemAutoReg::CheckAllTagsIsIn(TestTags, "bob,alice,foo,wiz") == false);
+	CHECK(FOnlineSubsystemAutoReg::CheckAllTagsIsIn(TestTags, "bob,wiz") == false);
+	CHECK(FOnlineSubsystemAutoReg::CheckAllTagsIsIn(TestTags, ",wiz") == false);
 
 	// Bound Checks
-	CHECK(OnlineSubsystemAutoReg::CheckAllTagsIsIn(TestTags, ",") == false);
-	CHECK(OnlineSubsystemAutoReg::CheckAllTagsIsIn(TestTags, "") == false);
-	CHECK(OnlineSubsystemAutoReg::CheckAllTagsIsIn({}, "") == false);
+	CHECK(FOnlineSubsystemAutoReg::CheckAllTagsIsIn(TestTags, ",") == false);
+	CHECK(FOnlineSubsystemAutoReg::CheckAllTagsIsIn(TestTags, "") == false);
+	CHECK(FOnlineSubsystemAutoReg::CheckAllTagsIsIn({}, "") == false);
 }
 
 
@@ -212,18 +212,18 @@ SELFTEST_TEST_CASE("CheckAllTagsIsIn(TArray, TArray) true cases", CHECKALLTAGSIS
 	CAPTURE(TestTags);
 
 	// Truthy Cases
-	CHECK(OnlineSubsystemAutoReg::CheckAllTagsIsIn(TestTags, { "bob", "alice" }) == true);
-	CHECK(OnlineSubsystemAutoReg::CheckAllTagsIsIn(TestTags, { "bob", "alice", "foo" }) == true);
-	CHECK(OnlineSubsystemAutoReg::CheckAllTagsIsIn(TestTags, TArray<FString>({ "foo" })) == true);
+	CHECK(FOnlineSubsystemAutoReg::CheckAllTagsIsIn(TestTags, { "bob", "alice" }) == true);
+	CHECK(FOnlineSubsystemAutoReg::CheckAllTagsIsIn(TestTags, { "bob", "alice", "foo" }) == true);
+	CHECK(FOnlineSubsystemAutoReg::CheckAllTagsIsIn(TestTags, TArray<FString>({ "foo" })) == true);
 
 	// Negative Cases
-	CHECK(OnlineSubsystemAutoReg::CheckAllTagsIsIn(TestTags, TArray<FString>({ "wiz" })) == false);
-	CHECK(OnlineSubsystemAutoReg::CheckAllTagsIsIn(TestTags, { "bob", "alice", "foo", "wiz" }) == false);
-	CHECK(OnlineSubsystemAutoReg::CheckAllTagsIsIn(TestTags, { "bob", "alice", "wiz" }) == false);
+	CHECK(FOnlineSubsystemAutoReg::CheckAllTagsIsIn(TestTags, TArray<FString>({ "wiz" })) == false);
+	CHECK(FOnlineSubsystemAutoReg::CheckAllTagsIsIn(TestTags, { "bob", "alice", "foo", "wiz" }) == false);
+	CHECK(FOnlineSubsystemAutoReg::CheckAllTagsIsIn(TestTags, { "bob", "alice", "wiz" }) == false);
 
 	// Bounds
-	CHECK(OnlineSubsystemAutoReg::CheckAllTagsIsIn(TestTags, TArray<FString>({})) == false);
-	CHECK(OnlineSubsystemAutoReg::CheckAllTagsIsIn({}, TArray<FString>({})) == false);
-	CHECK(OnlineSubsystemAutoReg::CheckAllTagsIsIn({}, TArray<FString>({ "wiz" })) == false);
-	CHECK(OnlineSubsystemAutoReg::CheckAllTagsIsIn({ "wiz" }, TArray<FString>({})) == false);
+	CHECK(FOnlineSubsystemAutoReg::CheckAllTagsIsIn(TestTags, TArray<FString>({})) == false);
+	CHECK(FOnlineSubsystemAutoReg::CheckAllTagsIsIn({}, TArray<FString>({})) == false);
+	CHECK(FOnlineSubsystemAutoReg::CheckAllTagsIsIn({}, TArray<FString>({ "wiz" })) == false);
+	CHECK(FOnlineSubsystemAutoReg::CheckAllTagsIsIn({ "wiz" }, TArray<FString>({})) == false);
 }
\ No newline at end of file
