Compile errors in IdentityLoginTests.cpp and Module.OSSTestsCore.cpp

- macro naming changes

diff --git a/Engine/Source/Programs/Online/OSSTestsCore/Tests/Identity/IdentityAutoLoginTests.cpp b/Engine/Source/Programs/Online/OSSTestsCore/Tests/Identity/IdentityAutoLoginTests.cpp
--- a/Engine/Source/Programs/Online/OSSTestsCore/Tests/Identity/IdentityAutoLoginTests.cpp
+++ b/Engine/Source/Programs/Online/OSSTestsCore/Tests/Identity/IdentityAutoLoginTests.cpp
@@ -15,9 +15,9 @@
 
 #define IDENTITY_TEST_CASE(x, ...) ONLINESUBSYSTEM_TEST_CASE(x, IDENTITY_TAG __VA_ARGS__)
 
-#define IDENTITY_LEGACY_EOS_TEST(x, ...) ONLINESUBSYSTEM_TEST_CASE_FIXTURE(FOnlineSubsystemEOSLegacyTestFixture, x, __VA_ARGS__)
+#define IDENTITY_LEGACY_AUTOLOGIN_EOS_TEST(x, ...) ONLINESUBSYSTEM_TEST_CASE_FIXTURE(FOnlineSubsystemEOSLegacyTestFixture, x, __VA_ARGS__)
 
-IDENTITY_LEGACY_EOS_TEST("Verify calling Identity AutoLogin with valid inputs and using legacy login flow returns the expected result(Success Case)", EG_IDENTITY_AUTOLOGIN_LEGACY_EOS_FLOW_TAG)
+IDENTITY_LEGACY_AUTOLOGIN_EOS_TEST("Verify calling Identity AutoLogin with valid inputs and using legacy login flow returns the expected result(Success Case)", EG_IDENTITY_AUTOLOGIN_LEGACY_EOS_FLOW_TAG)
 {
 	FTestDriver LocalDriver;
 	int32 LocalUserNum = 0;
diff --git a/Engine/Source/Programs/Online/OSSTestsCore/Tests/Identity/IdentityLoginTests.cpp b/Engine/Source/Programs/Online/OSSTestsCore/Tests/Identity/IdentityLoginTests.cpp
--- a/Engine/Source/Programs/Online/OSSTestsCore/Tests/Identity/IdentityLoginTests.cpp
+++ b/Engine/Source/Programs/Online/OSSTestsCore/Tests/Identity/IdentityLoginTests.cpp
@@ -11,9 +11,9 @@
 
 #define IDENTITY_TEST_CASE(x, ...) ONLINESUBSYSTEM_TEST_CASE(x, IDENTITY_TAG __VA_ARGS__)
 
-#define IDENTITY_LEGACY_EOS_TEST(x, ...) ONLINESUBSYSTEM_TEST_CASE_FIXTURE(FOnlineSubsystemEOSLegacyTestFixture, x, IDENTITY_TAG __VA_ARGS__)
+#define IDENTITY_LEGACY_LOGIN_EOS_TEST(x, ...) ONLINESUBSYSTEM_TEST_CASE_FIXTURE(FOnlineSubsystemEOSLegacyTestFixture, x, IDENTITY_TAG __VA_ARGS__)
 
-IDENTITY_LEGACY_EOS_TEST("Verify calling Identity Login with valid inputs and using legacy login flow returns the expected result(Success Case)", EG_IDENTITY_LOGIN_LEGACY_FLOWEOS_TAG)
+IDENTITY_LEGACY_LOGIN_EOS_TEST("Verify calling Identity Login with valid inputs and using legacy login flow returns the expected result(Success Case)", EG_IDENTITY_LOGIN_LEGACY_FLOWEOS_TAG)
 {
 	int32 NumUsersToImplicitLogin = 2;
 
