Added OnlineTests support for PS4/PS5 platforms.

diff --git a/Engine/Source/Programs/Online/OnlineTestsCore/Source/Private/Tests/Achievements/GetAchievementStateTests.cpp b/Engine/Source/Programs/Online/OnlineTestsCore/Source/Private/Tests/Achievements/GetAchievementStateTests.cpp
--- a/Engine/Source/Programs/Online/OnlineTestsCore/Source/Private/Tests/Achievements/GetAchievementStateTests.cpp
+++ b/Engine/Source/Programs/Online/OnlineTestsCore/Source/Private/Tests/Achievements/GetAchievementStateTests.cpp
@@ -8,6 +8,7 @@
 
 #define ACHIEVEMENTS_TAG "[suite_achievements]"
 #define EG_ACHIEVEMENTS_GETACHIEVEMENTSTATE_TAG ACHIEVEMENTS_TAG "[getachievementstate]"
+#define EG_ACHIEVEMENTS_GETACHIEVEMENTSTATE_SKIP_NOT_FOUND_TAG ACHIEVEMENTS_TAG "[getachievementstate][skip_achievement_not_found]"
 
 #define ACHIEVEMENTS_TEST_CASE(x, ...) ONLINE_TEST_CASE(x, ACHIEVEMENTS_TAG __VA_ARGS__)
 
@@ -67,7 +68,7 @@ ACHIEVEMENTS_TEST_CASE("Get Achievement State (Not Found)", EG_ACHIEVEMENTS_GETA
 	RunToCompletion();
 }
 
-ACHIEVEMENTS_TEST_CASE("Get Achievement State (Success)", EG_ACHIEVEMENTS_GETACHIEVEMENTSTATE_TAG)
+ACHIEVEMENTS_TEST_CASE("Get Achievement State (Success)", EG_ACHIEVEMENTS_GETACHIEVEMENTSTATE_SKIP_NOT_FOUND_TAG)
 {
 	FAccountId AccountId;
 
diff --git a/Engine/Source/Programs/Online/OnlineTestsCore/Source/Private/Tests/Sessions/AddRemoveSessionMemberTests.cpp b/Engine/Source/Programs/Online/OnlineTestsCore/Source/Private/Tests/Sessions/AddRemoveSessionMemberTests.cpp
--- a/Engine/Source/Programs/Online/OnlineTestsCore/Source/Private/Tests/Sessions/AddRemoveSessionMemberTests.cpp
+++ b/Engine/Source/Programs/Online/OnlineTestsCore/Source/Private/Tests/Sessions/AddRemoveSessionMemberTests.cpp
@@ -6,6 +6,7 @@
 
 #define SESSIONS_TAG "[suite_sessions]"
 #define EG_SESSIONS_ADDREMOVESESSIONMEMBER_TAG SESSIONS_TAG "[addremovesessionmember]"
+#define EG_SESSIONS_ADDREMOVESESSIONMEMBER_SKIP_TAG SESSIONS_TAG "[addremovesessionmember][skip_UserNumToLogin]"
 #define SESSIONS_TEST_CASE(x, ...) ONLINE_TEST_CASE(x, SESSIONS_TAG __VA_ARGS__)
 
 SESSIONS_TEST_CASE("If I call AddSessionMember with an invalid account id, I get an error", EG_SESSIONS_ADDREMOVESESSIONMEMBER_TAG)
@@ -153,7 +154,7 @@ SESSIONS_TEST_CASE("If I call RemoveSessionMember with an unregistered session n
 	RunToCompletion();
 }
 
-SESSIONS_TEST_CASE("If I call RemoveSessionMember with valid data, the operation completes successfully", EG_SESSIONS_ADDREMOVESESSIONMEMBER_TAG)
+SESSIONS_TEST_CASE("If I call RemoveSessionMember with valid data, the operation completes successfully", EG_SESSIONS_ADDREMOVESESSIONMEMBER_SKIP_TAG)
 {
 	DestroyCurrentServiceModule();
 
diff --git a/Engine/Source/Programs/Online/OnlineTestsCore/Source/Private/Tests/Sessions/CreateSessionTests.cpp b/Engine/Source/Programs/Online/OnlineTestsCore/Source/Private/Tests/Sessions/CreateSessionTests.cpp
--- a/Engine/Source/Programs/Online/OnlineTestsCore/Source/Private/Tests/Sessions/CreateSessionTests.cpp
+++ b/Engine/Source/Programs/Online/OnlineTestsCore/Source/Private/Tests/Sessions/CreateSessionTests.cpp
@@ -7,6 +7,7 @@
 
 #define SESSIONS_TAG "[suite_sessions]"
 #define EG_SESSIONS_CREATESESSION_TAG SESSIONS_TAG "[createsession]"
+#define EG_SESSIONS_CREATESESSION_SKIP_TAG SESSIONS_TAG "[createsession][skip_UserNumToLogin]"
 #define EG_SESSIONS_CREATESESSIONEOS_TAG SESSIONS_TAG "[createsession][.EOS]"
 #define EG_SESSIONS_CREATESESSIONNULL_TAG SESSIONS_TAG "[createsession][.NULL]"
 #define SESSIONS_TEST_CASE(x, ...) ONLINE_TEST_CASE(x, SESSIONS_TAG __VA_ARGS__)
@@ -146,7 +147,7 @@ SESSIONS_TEST_CASE("If I call CreateSession with session override id less than 1
 	RunToCompletion();
 }
 
-SESSIONS_TEST_CASE("If I call CreateSession with an name already in use, I get an error", EG_SESSIONS_CREATESESSION_TAG)
+SESSIONS_TEST_CASE("If I call CreateSession with an name already in use, I get an error", EG_SESSIONS_CREATESESSION_SKIP_TAG)
 {
 	DestroyCurrentServiceModule();
 	
@@ -187,7 +188,7 @@ SESSIONS_TEST_CASE("If I call CreateSession with an name already in use, I get a
 	RunToCompletion();
 }
 
-SESSIONS_TEST_CASE("If I call CreateSession with the presence flag set after there is already a presence session, I get an error", EG_SESSIONS_CREATESESSION_TAG)
+SESSIONS_TEST_CASE("If I call CreateSession with the presence flag set after there is already a presence session, I get an error", EG_SESSIONS_CREATESESSION_SKIP_TAG)
 {
 	LOG_SCOPE_VERBOSITY_OVERRIDE(LogOnlineServices, ELogVerbosity::NoLogging);
 
diff --git a/Engine/Source/Programs/Online/OnlineTestsCore/Source/Private/Tests/Sessions/JoinSessionTests.cpp b/Engine/Source/Programs/Online/OnlineTestsCore/Source/Private/Tests/Sessions/JoinSessionTests.cpp
--- a/Engine/Source/Programs/Online/OnlineTestsCore/Source/Private/Tests/Sessions/JoinSessionTests.cpp
+++ b/Engine/Source/Programs/Online/OnlineTestsCore/Source/Private/Tests/Sessions/JoinSessionTests.cpp
@@ -14,6 +14,7 @@
 
 #define SESSIONS_TAG "[suite_sessions]"
 #define EG_SESSIONS_JOINSESSIONS_TAG SESSIONS_TAG "[joinsession]"
+#define EG_SESSIONS_JOINSESSIONS_SKIP_TAG SESSIONS_TAG "[joinsession][skip_UserNumToLogin]"
 #define EG_SESSIONS_JOINSESSIONSEOS_TAG SESSIONS_TAG "[joinsession][.EOS]"
 #define SESSIONS_TEST_CASE(x, ...) ONLINE_TEST_CASE(x, SESSIONS_TAG __VA_ARGS__)
 
@@ -51,7 +52,7 @@ SESSIONS_TEST_CASE("If I call JoinSession with an invalid session id, I get an e
 	RunToCompletion();
 }
 
-SESSIONS_TEST_CASE("If I call JoinSession with an empty session name, I get an error", EG_SESSIONS_JOINSESSIONS_TAG)
+SESSIONS_TEST_CASE("If I call JoinSession with an empty session name, I get an error", EG_SESSIONS_JOINSESSIONS_SKIP_TAG)
 {
 	LOG_SCOPE_VERBOSITY_OVERRIDE(LogEOSSDK, ELogVerbosity::NoLogging);
 
@@ -102,7 +103,7 @@ SESSIONS_TEST_CASE("If I call JoinSession with an empty session name, I get an e
 	RunToCompletion();
 }
 
-SESSIONS_TEST_CASE("If I call JoinSession with an name already in use, I get an error", EG_SESSIONS_JOINSESSIONSEOS_TAG)
+SESSIONS_TEST_CASE("If I call JoinSession with an name already in use, I get an error", EG_SESSIONS_JOINSESSIONS_SKIP_TAG)
 {
 	LOG_SCOPE_VERBOSITY_OVERRIDE(LogEOSSDK, ELogVerbosity::NoLogging);
 
diff --git a/Engine/Source/Programs/Online/OnlineTestsCore/Source/Private/Tests/Sessions/LeaveSessionTests.cpp b/Engine/Source/Programs/Online/OnlineTestsCore/Source/Private/Tests/Sessions/LeaveSessionTests.cpp
--- a/Engine/Source/Programs/Online/OnlineTestsCore/Source/Private/Tests/Sessions/LeaveSessionTests.cpp
+++ b/Engine/Source/Programs/Online/OnlineTestsCore/Source/Private/Tests/Sessions/LeaveSessionTests.cpp
@@ -5,6 +5,7 @@
 
 #define SESSIONS_TAG "[suite_sessions]"
 #define SESSIONS_LEAVESESSION_TAG SESSIONS_TAG "[leavesesssion]"
+#define SESSIONS_LEAVESESSION_SKIP_TAG SESSIONS_TAG "[leavesesssion][skip_UserNumToLogin]"
 #define SESSIONS_TEST_CASE(x, ...) ONLINE_TEST_CASE(x, SESSIONS_TAG __VA_ARGS__)
 
 SESSIONS_TEST_CASE("If I call LeaveSession with an invalid account id, I get an error", SESSIONS_LEAVESESSION_TAG)
@@ -61,7 +62,7 @@ SESSIONS_TEST_CASE("If I call LeaveSession with an unregistered session name, I
 	RunToCompletion();
 }
 
-SESSIONS_TEST_CASE("If I call LeaveSession with valid data, the operation completes successfully", SESSIONS_LEAVESESSION_TAG)
+SESSIONS_TEST_CASE("If I call LeaveSession with valid data, the operation completes successfully", SESSIONS_LEAVESESSION_SKIP_TAG)
 {
 	DestroyCurrentServiceModule();
 
diff --git a/Engine/Source/Programs/Online/OnlineTestsCore/Source/Private/Tests/Sessions/PresenceSessionTests.cpp b/Engine/Source/Programs/Online/OnlineTestsCore/Source/Private/Tests/Sessions/PresenceSessionTests.cpp
--- a/Engine/Source/Programs/Online/OnlineTestsCore/Source/Private/Tests/Sessions/PresenceSessionTests.cpp
+++ b/Engine/Source/Programs/Online/OnlineTestsCore/Source/Private/Tests/Sessions/PresenceSessionTests.cpp
@@ -9,6 +9,7 @@
 
 #define SESSIONS_TAG "[suite_sessions]"
 #define EG_SESSIONS_PRESENCESESSION_TAG SESSIONS_TAG "[presencesession]"
+#define EG_SESSIONS_PRESENCESESSION_SKIP_TAG SESSIONS_TAG "[presencesession][skip_UserNumToLogin]"
 #define EG_SESSIONS_PRESENCESESSIONEOS_TAG SESSIONS_TAG "[presencesession][.EOS]"
 #define EG_SESSIONS_PRESENCESESSIONNULL_TAG SESSIONS_TAG "[presencesession][.NULL]"
 #define SESSIONS_TEST_CASE(x, ...) ONLINE_TEST_CASE(x, SESSIONS_TAG __VA_ARGS__)
@@ -52,7 +53,7 @@ SESSIONS_TEST_CASE("If I call IsPresenceSession with an invalid session id, I ge
 	RunToCompletion();
 }
 
-SESSIONS_TEST_CASE("If I call IsPresenceSession with an unregistered account id, I get an error", EG_SESSIONS_PRESENCESESSION_TAG)
+SESSIONS_TEST_CASE("If I call IsPresenceSession with an unregistered account id, I get an error", EG_SESSIONS_PRESENCESESSION_SKIP_TAG)
 {
 	LOG_SCOPE_VERBOSITY_OVERRIDE(LogOnlineServices, ELogVerbosity::NoLogging);
 
diff --git a/Engine/Source/Programs/Online/OnlineTestsCore/Source/Private/Tests/Sessions/SendRejectSessionInvitesTests.cpp b/Engine/Source/Programs/Online/OnlineTestsCore/Source/Private/Tests/Sessions/SendRejectSessionInvitesTests.cpp
--- a/Engine/Source/Programs/Online/OnlineTestsCore/Source/Private/Tests/Sessions/SendRejectSessionInvitesTests.cpp
+++ b/Engine/Source/Programs/Online/OnlineTestsCore/Source/Private/Tests/Sessions/SendRejectSessionInvitesTests.cpp
@@ -9,6 +9,7 @@
 
 #define SESSIONS_TAG "[suite_sessions]"
 #define EG_SESSIONS_SENDREJECTSESSIONINVITE_TAG SESSIONS_TAG "[sendrejectsesssioninvite]"
+#define EG_SESSIONS_SENDREJECTSESSIONINVITE_SKIP_TAG SESSIONS_TAG "[sendrejectsesssioninvite][skip_UserNumToLogin]"
 #define EG_SESSIONS_SENDREJECTSESSIONINVITEEOS_TAG SESSIONS_TAG "[sendrejectsesssioninvite][.EOS]"
 #define SESSIONS_TEST_CASE(x, ...) ONLINE_TEST_CASE(x, SESSIONS_TAG __VA_ARGS__)
 
@@ -49,7 +50,7 @@ SESSIONS_TEST_CASE("If I call SendSessionInvite with an empty session name, I ge
 	RunToCompletion();
 }
 
-SESSIONS_TEST_CASE("If I call SendSessionInvite with a valid session name but unregistered, I get an error", EG_SESSIONS_SENDREJECTSESSIONINVITE_TAG)
+SESSIONS_TEST_CASE("If I call SendSessionInvite with a valid session name but unregistered, I get an error", EG_SESSIONS_SENDREJECTSESSIONINVITE_SKIP_TAG)
 {
 	FAccountId FirstAccountId, SecondAccountId;
 
@@ -70,7 +71,7 @@ SESSIONS_TEST_CASE("If I call SendSessionInvite with a valid session name but un
 	RunToCompletion();
 }
 
-SESSIONS_TEST_CASE("If I call SendSessionInvite with an empty target users, I get an error", EG_SESSIONS_SENDREJECTSESSIONINVITE_TAG)
+SESSIONS_TEST_CASE("If I call SendSessionInvite with an empty target users, I get an error", EG_SESSIONS_SENDREJECTSESSIONINVITE_SKIP_TAG)
 {
 	DestroyCurrentServiceModule();
 
@@ -111,7 +112,7 @@ SESSIONS_TEST_CASE("If I call SendSessionInvite with an empty target users, I ge
 	RunToCompletion();
 }
 
-SESSIONS_TEST_CASE("If I call SendSessionInvite with an invalid target users, I get an error", EG_SESSIONS_SENDREJECTSESSIONINVITE_TAG)
+SESSIONS_TEST_CASE("If I call SendSessionInvite with an invalid target users, I get an error", EG_SESSIONS_SENDREJECTSESSIONINVITE_SKIP_TAG)
 {
 	DestroyCurrentServiceModule();
 
diff --git a/Engine/Source/Programs/Online/OnlineTestsCore/Source/Private/Tests/Sessions/StartMatchmakingTests.cpp b/Engine/Source/Programs/Online/OnlineTestsCore/Source/Private/Tests/Sessions/StartMatchmakingTests.cpp
--- a/Engine/Source/Programs/Online/OnlineTestsCore/Source/Private/Tests/Sessions/StartMatchmakingTests.cpp
+++ b/Engine/Source/Programs/Online/OnlineTestsCore/Source/Private/Tests/Sessions/StartMatchmakingTests.cpp
@@ -8,6 +8,7 @@
 
 #define SESSIONS_TAG "[suite_sessions]"
 #define EG_SESSIONS_STARTMATCHMAKING_TAG SESSIONS_TAG "[startmatchmaking]"
+#define EG_SESSIONS_STARTMATCHMAKING_SKIP_TAG SESSIONS_TAG "[startmatchmaking][skip_UserNumToLogin]"
 #define EG_SESSIONS_STARTMATCHMAKINGEOS_TAG SESSIONS_TAG "[startmatchmaking][.EOS]"
 #define EG_SESSIONS_STARTMATCHMAKINGEOSNULL_TAG SESSIONS_TAG "[startmatchmaking][.EOS][.NULL]"
 #define SESSIONS_TEST_CASE(x, ...) ONLINE_TEST_CASE(x, SESSIONS_TAG __VA_ARGS__)
@@ -173,7 +174,7 @@ SESSIONS_TEST_CASE("If I call StartMatchmaking with session override id less tha
 	RunToCompletion();
 }
 
-SESSIONS_TEST_CASE("If I call StartMatchmaking with name already in use, I get an error", EG_SESSIONS_STARTMATCHMAKING_TAG)
+SESSIONS_TEST_CASE("If I call StartMatchmaking with name already in use, I get an error", EG_SESSIONS_STARTMATCHMAKING_SKIP_TAG)
 {
 	DestroyCurrentServiceModule();
 
diff --git a/Engine/Source/Programs/Online/OnlineTestsCore/Source/Private/Tests/Sessions/UpdateSessionSettingsTests.cpp b/Engine/Source/Programs/Online/OnlineTestsCore/Source/Private/Tests/Sessions/UpdateSessionSettingsTests.cpp
--- a/Engine/Source/Programs/Online/OnlineTestsCore/Source/Private/Tests/Sessions/UpdateSessionSettingsTests.cpp
+++ b/Engine/Source/Programs/Online/OnlineTestsCore/Source/Private/Tests/Sessions/UpdateSessionSettingsTests.cpp
@@ -7,6 +7,7 @@
 
 #define SESSIONS_TAGS "[suite_sessions]"
 #define SESSIONS_UPDATESESSIONSETTINGS_TAG SESSIONS_TAGS "[updatesessionsettings]"
+#define SESSIONS_UPDATESESSIONSETTINGS_SKIP_TAG SESSIONS_TAGS "[updatesessionsettings][skip_UserNumToLogin]"
 #define UPDATESESSIONS_TEST_CASE(x, ...) ONLINE_TEST_CASE(x, SESSIONS_TAGS __VA_ARGS__)
 
 UPDATESESSIONS_TEST_CASE("If I call UpdateSessionSettings with an invalid account id, I get an error", SESSIONS_UPDATESESSIONSETTINGS_TAG)
@@ -126,7 +127,7 @@ UPDATESESSIONS_TEST_CASE("If I call UpdateSessionSettings with an unregistered s
 	RunToCompletion();
 }
 
-UPDATESESSIONS_TEST_CASE("If I call UpdateSessionSettings with valid data, the operation completes successfully", SESSIONS_UPDATESESSIONSETTINGS_TAG)
+UPDATESESSIONS_TEST_CASE("If I call UpdateSessionSettings with valid data, the operation completes successfully", SESSIONS_UPDATESESSIONSETTINGS_SKIP_TAG)
 {
 	DestroyCurrentServiceModule();
 
diff --git a/Engine/Source/Programs/Online/OnlineTestsCore/Source/Private/Tests/TestAsyncOp.cpp b/Engine/Source/Programs/Online/OnlineTestsCore/Source/Private/Tests/TestAsyncOp.cpp
--- a/Engine/Source/Programs/Online/OnlineTestsCore/Source/Private/Tests/TestAsyncOp.cpp
+++ b/Engine/Source/Programs/Online/OnlineTestsCore/Source/Private/Tests/TestAsyncOp.cpp
@@ -14,6 +14,7 @@
 #include "HAL/PlatformTLS.h"
 
 #define ASYNCOP_TAG "[AsyncOp]"
+#define ASYNCOP_SKIP_EOS_PLATFORM_FAULT_TAG "[AsyncOp][skip_eos_platform_fault]"
 #define ASYNCOP_TEST_CASE(x, ...) ONLINE_TEST_CASE(x, ASYNCOP_TAG __VA_ARGS__)
 
 struct FTestOp
@@ -306,7 +307,7 @@ ASYNCOP_TEST_CASE("Verify that the duration cache policy works)")
 	RunToCompletion();
 }
 
-ASYNCOP_TEST_CASE("Verify that the cache error setting works)")
+ASYNCOP_TEST_CASE("Verify that the cache error setting works)", ASYNCOP_SKIP_EOS_PLATFORM_FAULT_TAG)
 {
 	DestroyCurrentServiceModule();
 	FTestPipeline& TestPipeline = GetPipeline();
diff --git a/Engine/Source/Programs/Online/OnlineTestsCore/Source/Private/Tests/TitleFile/TitleFileTests.cpp b/Engine/Source/Programs/Online/OnlineTestsCore/Source/Private/Tests/TitleFile/TitleFileTests.cpp
--- a/Engine/Source/Programs/Online/OnlineTestsCore/Source/Private/Tests/TitleFile/TitleFileTests.cpp
+++ b/Engine/Source/Programs/Online/OnlineTestsCore/Source/Private/Tests/TitleFile/TitleFileTests.cpp
@@ -11,6 +11,7 @@
 
 #define TITLE_FILE_ENUMERATE_TAG "[enumeratefiles]"
 #define TITLE_FILE_GET_ENUMERATED_TAG "[getenumeratedfiles]"
+#define TITLE_FILE_GET_ENUMERATED_SKIP_EOS_PLATFORM_FAULT_TAG "[getenumeratedfiles][skip_eos_platform_fault]"
 #define TITLE_FILE_READ_TAG "[readfile]"
 
 // These variables must match what is defined in the backend for every implementation that the test is going to run on
@@ -87,7 +88,7 @@ TITLE_FILE_TEST_CASE("GetEnumeratedFiles - Invalid user", TITLE_FILE_GET_ENUMERA
 	RunToCompletion();
 }
 
-TITLE_FILE_TEST_CASE("GetEnumeratedFiles - Invalid state", TITLE_FILE_GET_ENUMERATED_TAG)
+TITLE_FILE_TEST_CASE("GetEnumeratedFiles - Invalid state", TITLE_FILE_GET_ENUMERATED_SKIP_EOS_PLATFORM_FAULT_TAG)
 {
 	DestroyCurrentServiceModule();
 
