OnlineTests: Fixed local user IDs (PlatformUserId) mismatch.
Disabled MultiAccount tests for the PS4/PS5.

diff --git a/Engine/Plugins/Online/OnlineServicesEpicCommon/Source/OnlineServicesEpicCommon/Private/Online/OnlineErrorEpicCommon.cpp b/Engine/Plugins/Online/OnlineServicesEpicCommon/Source/OnlineServicesEpicCommon/Private/Online/OnlineErrorEpicCommon.cpp
--- a/Engine/Plugins/Online/OnlineServicesEpicCommon/Source/OnlineServicesEpicCommon/Private/Online/OnlineErrorEpicCommon.cpp
+++ b/Engine/Plugins/Online/OnlineServicesEpicCommon/Source/OnlineServicesEpicCommon/Private/Online/OnlineErrorEpicCommon.cpp
@@ -65,6 +65,7 @@ FOnlineError MapCommonEOSError(FOnlineError&& Error, EOS_EResult Result)
 	case EOS_EResult::EOS_Canceled:				return Errors::Cancelled(MoveTemp(Error));
 	case EOS_EResult::EOS_NotFound:				return Errors::NotFound(MoveTemp(Error));
 	case EOS_EResult::EOS_OperationWillRetry:	return Errors::WillRetry(MoveTemp(Error));
+	case EOS_EResult::EOS_NoChange:				return Errors::NoChange(MoveTemp(Error));
 	case EOS_EResult::EOS_VersionMismatch:		return Errors::IncompatibleVersion(MoveTemp(Error));
 	case EOS_EResult::EOS_LimitExceeded:		return Errors::TooManyRequests(MoveTemp(Error));
 	case EOS_EResult::EOS_TimedOut:				return Errors::Timeout(MoveTemp(Error));
diff --git a/Engine/Source/Programs/Online/OnlineTestsCore/Source/Private/OnlineCatchHelper.cpp b/Engine/Source/Programs/Online/OnlineTestsCore/Source/Private/OnlineCatchHelper.cpp
--- a/Engine/Source/Programs/Online/OnlineTestsCore/Source/Private/OnlineCatchHelper.cpp
+++ b/Engine/Source/Programs/Online/OnlineTestsCore/Source/Private/OnlineCatchHelper.cpp
@@ -19,15 +19,14 @@
 
 // Make sure there are registered input devices for N users and fire
 // OnInputDeviceConnectionChange delegate for interested online service code.
-void EnsureLocalUserCount(uint32 UserNumToLogin, uint32 NumUsers)
+void EnsureLocalUserCount(uint32 NumUsers)
 {
 	TArray<FPlatformUserId> Users;
 	IPlatformInputDeviceMapper::Get().GetAllActiveUsers(Users);
 
-	int32 Index = UserNumToLogin * 10;
-	const int32 EndIndex = NumUsers + Index;
+	const uint32 ActiveUserCount = Users.Num();
 
-	for (; Index < EndIndex; ++Index)
+	for (uint32 Index = ActiveUserCount; Index < NumUsers; ++Index)
 	{
 		int32 Result = Users.Find(FPlatformMisc::GetPlatformUserForUserIndex(Index));
 		if (Result < 0)
@@ -167,7 +166,7 @@ TArray<UE::Online::FAuthLogin::Params> OnlineTestBase::GetIniCredentials(int32 L
 	{
 		Params.CredentialsToken.Set<FString>(LoginCredentialSplit[2]);
 	}
-	Params.PlatformUserId = FPlatformMisc::GetPlatformUserForUserIndex(LocalUserNum);
+	Params.PlatformUserId = FPlatformMisc::GetPlatformUserForUserIndex(AuthLoginParams.Num());
 
 	AuthLoginParams.Add(Params);
 
@@ -215,7 +214,7 @@ FTestPipeline& OnlineTestBase::GetLoginPipeline(uint32 UserNumToLogin, std::init
 	GConfig->GetBool(*LoginCredentialCategory, TEXT("UseImplicitLogin"), bUseImplicitLogin, GEngineIni);
 	
 	// Make sure input delegates are fired for adding the required user count.
-	EnsureLocalUserCount(UserNumToLogin, NumLocalUsers);
+	EnsureLocalUserCount(NumLocalUsers);
 
 	if (bUseImplicitLogin)
 	{
@@ -286,7 +285,8 @@ void OnlineTestBase::RunToCompletion(bool bLogout) const
 		}
 		else if(NumLocalUsers > 0)
 		{
-			for(UE::Online::FAuthLogin::Params AuthLoginParam : GetCredentials(UserNumToLogout, NumLocalUsers))
+			TArray<UE::Online::FAuthLogin::Params> AuthLoginParams = GetCredentials(UserNumToLogout, NumLocalUsers);
+			for (UE::Online::FAuthLogin::Params AuthLoginParam : AuthLoginParams)
 			{
 				Pipeline->EmplaceStep<FAuthLogoutStep>(AuthLoginParam.PlatformUserId);
 			}
diff --git a/Engine/Source/Programs/Online/OnlineTestsCore/Source/Private/Tests/Sessions/AddRemoveSessionMemberTests.cpp b/Engine/Source/Programs/Online/OnlineTestsCore/Source/Private/Tests/Sessions/AddRemoveSessionMemberTests.cpp
--- a/Engine/Source/Programs/Online/OnlineTestsCore/Source/Private/Tests/Sessions/AddRemoveSessionMemberTests.cpp
+++ b/Engine/Source/Programs/Online/OnlineTestsCore/Source/Private/Tests/Sessions/AddRemoveSessionMemberTests.cpp
@@ -6,7 +6,6 @@
 
 #define SESSIONS_TAG "[suite_sessions]"
 #define EG_SESSIONS_ADDREMOVESESSIONMEMBER_TAG SESSIONS_TAG "[addremovesessionmember]"
-#define EG_SESSIONS_ADDREMOVESESSIONMEMBER_SKIP_TAG SESSIONS_TAG "[addremovesessionmember][skip_UserNumToLogin]"
 #define SESSIONS_TEST_CASE(x, ...) ONLINE_TEST_CASE(x, SESSIONS_TAG __VA_ARGS__)
 
 SESSIONS_TEST_CASE("If I call AddSessionMember with an invalid account id, I get an error", EG_SESSIONS_ADDREMOVESESSIONMEMBER_TAG)
@@ -154,7 +153,7 @@ SESSIONS_TEST_CASE("If I call RemoveSessionMember with an unregistered session n
 	RunToCompletion();
 }
 
-SESSIONS_TEST_CASE("If I call RemoveSessionMember with valid data, the operation completes successfully", EG_SESSIONS_ADDREMOVESESSIONMEMBER_SKIP_TAG)
+SESSIONS_TEST_CASE("If I call RemoveSessionMember with valid data, the operation completes successfully", EG_SESSIONS_ADDREMOVESESSIONMEMBER_TAG)
 {
 	DestroyCurrentServiceModule();
 
diff --git a/Engine/Source/Programs/Online/OnlineTestsCore/Source/Private/Tests/Sessions/CreateSessionTests.cpp b/Engine/Source/Programs/Online/OnlineTestsCore/Source/Private/Tests/Sessions/CreateSessionTests.cpp
--- a/Engine/Source/Programs/Online/OnlineTestsCore/Source/Private/Tests/Sessions/CreateSessionTests.cpp
+++ b/Engine/Source/Programs/Online/OnlineTestsCore/Source/Private/Tests/Sessions/CreateSessionTests.cpp
@@ -7,7 +7,6 @@
 
 #define SESSIONS_TAG "[suite_sessions]"
 #define EG_SESSIONS_CREATESESSION_TAG SESSIONS_TAG "[createsession]"
-#define EG_SESSIONS_CREATESESSION_SKIP_TAG SESSIONS_TAG "[createsession][skip_UserNumToLogin]"
 #define EG_SESSIONS_CREATESESSIONEOS_TAG SESSIONS_TAG "[createsession][.EOS]"
 #define EG_SESSIONS_CREATESESSIONNULL_TAG SESSIONS_TAG "[createsession][.NULL]"
 #define SESSIONS_TEST_CASE(x, ...) ONLINE_TEST_CASE(x, SESSIONS_TAG __VA_ARGS__)
@@ -147,7 +146,7 @@ SESSIONS_TEST_CASE("If I call CreateSession with session override id less than 1
 	RunToCompletion();
 }
 
-SESSIONS_TEST_CASE("If I call CreateSession with an name already in use, I get an error", EG_SESSIONS_CREATESESSION_SKIP_TAG)
+SESSIONS_TEST_CASE("If I call CreateSession with an name already in use, I get an error", EG_SESSIONS_CREATESESSION_TAG)
 {
 	DestroyCurrentServiceModule();
 	
@@ -188,7 +187,7 @@ SESSIONS_TEST_CASE("If I call CreateSession with an name already in use, I get a
 	RunToCompletion();
 }
 
-SESSIONS_TEST_CASE("If I call CreateSession with the presence flag set after there is already a presence session, I get an error", EG_SESSIONS_CREATESESSION_SKIP_TAG)
+SESSIONS_TEST_CASE("If I call CreateSession with the presence flag set after there is already a presence session, I get an error", EG_SESSIONS_CREATESESSION_TAG)
 {
 	LOG_SCOPE_VERBOSITY_OVERRIDE(LogOnlineServices, ELogVerbosity::NoLogging);
 
diff --git a/Engine/Source/Programs/Online/OnlineTestsCore/Source/Private/Tests/Sessions/JoinSessionTests.cpp b/Engine/Source/Programs/Online/OnlineTestsCore/Source/Private/Tests/Sessions/JoinSessionTests.cpp
--- a/Engine/Source/Programs/Online/OnlineTestsCore/Source/Private/Tests/Sessions/JoinSessionTests.cpp
+++ b/Engine/Source/Programs/Online/OnlineTestsCore/Source/Private/Tests/Sessions/JoinSessionTests.cpp
@@ -14,7 +14,6 @@
 
 #define SESSIONS_TAG "[suite_sessions]"
 #define EG_SESSIONS_JOINSESSIONS_TAG SESSIONS_TAG "[joinsession]"
-#define EG_SESSIONS_JOINSESSIONS_SKIP_TAG SESSIONS_TAG "[joinsession][skip_UserNumToLogin]"
 #define EG_SESSIONS_JOINSESSIONSEOS_TAG SESSIONS_TAG "[joinsession][.EOS]"
 #define SESSIONS_TEST_CASE(x, ...) ONLINE_TEST_CASE(x, SESSIONS_TAG __VA_ARGS__)
 
@@ -52,7 +51,7 @@ SESSIONS_TEST_CASE("If I call JoinSession with an invalid session id, I get an e
 	RunToCompletion();
 }
 
-SESSIONS_TEST_CASE("If I call JoinSession with an empty session name, I get an error", EG_SESSIONS_JOINSESSIONS_SKIP_TAG)
+SESSIONS_TEST_CASE("If I call JoinSession with an empty session name, I get an error", EG_SESSIONS_JOINSESSIONS_TAG)
 {
 	LOG_SCOPE_VERBOSITY_OVERRIDE(LogEOSSDK, ELogVerbosity::NoLogging);
 
@@ -103,7 +102,7 @@ SESSIONS_TEST_CASE("If I call JoinSession with an empty session name, I get an e
 	RunToCompletion();
 }
 
-SESSIONS_TEST_CASE("If I call JoinSession with an name already in use, I get an error", EG_SESSIONS_JOINSESSIONS_SKIP_TAG)
+SESSIONS_TEST_CASE("If I call JoinSession with an name already in use, I get an error", EG_SESSIONS_JOINSESSIONSEOS_TAG)
 {
 	LOG_SCOPE_VERBOSITY_OVERRIDE(LogEOSSDK, ELogVerbosity::NoLogging);
 
diff --git a/Engine/Source/Programs/Online/OnlineTestsCore/Source/Private/Tests/Sessions/LeaveSessionTests.cpp b/Engine/Source/Programs/Online/OnlineTestsCore/Source/Private/Tests/Sessions/LeaveSessionTests.cpp
--- a/Engine/Source/Programs/Online/OnlineTestsCore/Source/Private/Tests/Sessions/LeaveSessionTests.cpp
+++ b/Engine/Source/Programs/Online/OnlineTestsCore/Source/Private/Tests/Sessions/LeaveSessionTests.cpp
@@ -5,7 +5,6 @@
 
 #define SESSIONS_TAG "[suite_sessions]"
 #define SESSIONS_LEAVESESSION_TAG SESSIONS_TAG "[leavesesssion]"
-#define SESSIONS_LEAVESESSION_SKIP_TAG SESSIONS_TAG "[leavesesssion][skip_UserNumToLogin]"
 #define SESSIONS_TEST_CASE(x, ...) ONLINE_TEST_CASE(x, SESSIONS_TAG __VA_ARGS__)
 
 SESSIONS_TEST_CASE("If I call LeaveSession with an invalid account id, I get an error", SESSIONS_LEAVESESSION_TAG)
@@ -62,7 +61,7 @@ SESSIONS_TEST_CASE("If I call LeaveSession with an unregistered session name, I
 	RunToCompletion();
 }
 
-SESSIONS_TEST_CASE("If I call LeaveSession with valid data, the operation completes successfully", SESSIONS_LEAVESESSION_SKIP_TAG)
+SESSIONS_TEST_CASE("If I call LeaveSession with valid data, the operation completes successfully", SESSIONS_LEAVESESSION_TAG)
 {
 	DestroyCurrentServiceModule();
 
diff --git a/Engine/Source/Programs/Online/OnlineTestsCore/Source/Private/Tests/Sessions/PresenceSessionTests.cpp b/Engine/Source/Programs/Online/OnlineTestsCore/Source/Private/Tests/Sessions/PresenceSessionTests.cpp
--- a/Engine/Source/Programs/Online/OnlineTestsCore/Source/Private/Tests/Sessions/PresenceSessionTests.cpp
+++ b/Engine/Source/Programs/Online/OnlineTestsCore/Source/Private/Tests/Sessions/PresenceSessionTests.cpp
@@ -9,7 +9,6 @@
 
 #define SESSIONS_TAG "[suite_sessions]"
 #define EG_SESSIONS_PRESENCESESSION_TAG SESSIONS_TAG "[presencesession]"
-#define EG_SESSIONS_PRESENCESESSION_SKIP_TAG SESSIONS_TAG "[presencesession][skip_UserNumToLogin]"
 #define EG_SESSIONS_PRESENCESESSIONEOS_TAG SESSIONS_TAG "[presencesession][.EOS]"
 #define EG_SESSIONS_PRESENCESESSIONNULL_TAG SESSIONS_TAG "[presencesession][.NULL]"
 #define SESSIONS_TEST_CASE(x, ...) ONLINE_TEST_CASE(x, SESSIONS_TAG __VA_ARGS__)
@@ -53,7 +52,7 @@ SESSIONS_TEST_CASE("If I call IsPresenceSession with an invalid session id, I ge
 	RunToCompletion();
 }
 
-SESSIONS_TEST_CASE("If I call IsPresenceSession with an unregistered account id, I get an error", EG_SESSIONS_PRESENCESESSION_SKIP_TAG)
+SESSIONS_TEST_CASE("If I call IsPresenceSession with an unregistered account id, I get an error", EG_SESSIONS_PRESENCESESSION_TAG)
 {
 	LOG_SCOPE_VERBOSITY_OVERRIDE(LogOnlineServices, ELogVerbosity::NoLogging);
 
diff --git a/Engine/Source/Programs/Online/OnlineTestsCore/Source/Private/Tests/Sessions/SendRejectSessionInvitesTests.cpp b/Engine/Source/Programs/Online/OnlineTestsCore/Source/Private/Tests/Sessions/SendRejectSessionInvitesTests.cpp
--- a/Engine/Source/Programs/Online/OnlineTestsCore/Source/Private/Tests/Sessions/SendRejectSessionInvitesTests.cpp
+++ b/Engine/Source/Programs/Online/OnlineTestsCore/Source/Private/Tests/Sessions/SendRejectSessionInvitesTests.cpp
@@ -9,8 +9,8 @@
 
 #define SESSIONS_TAG "[suite_sessions]"
 #define EG_SESSIONS_SENDREJECTSESSIONINVITE_TAG SESSIONS_TAG "[sendrejectsesssioninvite]"
-#define EG_SESSIONS_SENDREJECTSESSIONINVITE_SKIP_TAG SESSIONS_TAG "[sendrejectsesssioninvite][skip_UserNumToLogin]"
 #define EG_SESSIONS_SENDREJECTSESSIONINVITEEOS_TAG SESSIONS_TAG "[sendrejectsesssioninvite][.EOS]"
+#define EG_SESSIONS_SENDREJECTSESSIONINVITE_MULTIACCOUNT_TAG SESSIONS_TAG "[sendrejectsesssioninvite][MultiAccount]"
 #define SESSIONS_TEST_CASE(x, ...) ONLINE_TEST_CASE(x, SESSIONS_TAG __VA_ARGS__)
 
 SESSIONS_TEST_CASE("If I call SendSessionInvite with an invalid account id, I get an error", EG_SESSIONS_SENDREJECTSESSIONINVITE_TAG)
@@ -50,7 +50,7 @@ SESSIONS_TEST_CASE("If I call SendSessionInvite with an empty session name, I ge
 	RunToCompletion();
 }
 
-SESSIONS_TEST_CASE("If I call SendSessionInvite with a valid session name but unregistered, I get an error", EG_SESSIONS_SENDREJECTSESSIONINVITE_SKIP_TAG)
+SESSIONS_TEST_CASE("If I call SendSessionInvite with a valid session name but unregistered, I get an error", EG_SESSIONS_SENDREJECTSESSIONINVITE_MULTIACCOUNT_TAG)
 {
 	FAccountId FirstAccountId, SecondAccountId;
 
@@ -71,7 +71,7 @@ SESSIONS_TEST_CASE("If I call SendSessionInvite with a valid session name but un
 	RunToCompletion();
 }
 
-SESSIONS_TEST_CASE("If I call SendSessionInvite with an empty target users, I get an error", EG_SESSIONS_SENDREJECTSESSIONINVITE_SKIP_TAG)
+SESSIONS_TEST_CASE("If I call SendSessionInvite with an empty target users, I get an error", EG_SESSIONS_SENDREJECTSESSIONINVITE_MULTIACCOUNT_TAG)
 {
 	DestroyCurrentServiceModule();
 
@@ -112,7 +112,7 @@ SESSIONS_TEST_CASE("If I call SendSessionInvite with an empty target users, I ge
 	RunToCompletion();
 }
 
-SESSIONS_TEST_CASE("If I call SendSessionInvite with an invalid target users, I get an error", EG_SESSIONS_SENDREJECTSESSIONINVITE_SKIP_TAG)
+SESSIONS_TEST_CASE("If I call SendSessionInvite with an invalid target users, I get an error", EG_SESSIONS_SENDREJECTSESSIONINVITE_MULTIACCOUNT_TAG)
 {
 	DestroyCurrentServiceModule();
 
diff --git a/Engine/Source/Programs/Online/OnlineTestsCore/Source/Private/Tests/Sessions/StartMatchmakingTests.cpp b/Engine/Source/Programs/Online/OnlineTestsCore/Source/Private/Tests/Sessions/StartMatchmakingTests.cpp
--- a/Engine/Source/Programs/Online/OnlineTestsCore/Source/Private/Tests/Sessions/StartMatchmakingTests.cpp
+++ b/Engine/Source/Programs/Online/OnlineTestsCore/Source/Private/Tests/Sessions/StartMatchmakingTests.cpp
@@ -8,9 +8,8 @@
 
 #define SESSIONS_TAG "[suite_sessions]"
 #define EG_SESSIONS_STARTMATCHMAKING_TAG SESSIONS_TAG "[startmatchmaking]"
-#define EG_SESSIONS_STARTMATCHMAKING_SKIP_TAG SESSIONS_TAG "[startmatchmaking][skip_UserNumToLogin]"
 #define EG_SESSIONS_STARTMATCHMAKINGEOS_TAG SESSIONS_TAG "[startmatchmaking][.EOS]"
-#define EG_SESSIONS_STARTMATCHMAKINGEOSNULL_TAG SESSIONS_TAG "[startmatchmaking][.EOS][.NULL]"
+#define EG_SESSIONS_STARTMATCHMAKING_SKIP_TAG SESSIONS_TAG "[startmatchmaking][startmatchmakingskip]"
 #define SESSIONS_TEST_CASE(x, ...) ONLINE_TEST_CASE(x, SESSIONS_TAG __VA_ARGS__)
 
 SESSIONS_TEST_CASE("If I call StartMatchmaking with an invalid account id, I get an error", EG_SESSIONS_STARTMATCHMAKING_TAG)
@@ -174,7 +173,7 @@ SESSIONS_TEST_CASE("If I call StartMatchmaking with session override id less tha
 	RunToCompletion();
 }
 
-SESSIONS_TEST_CASE("If I call StartMatchmaking with name already in use, I get an error", EG_SESSIONS_STARTMATCHMAKING_SKIP_TAG)
+SESSIONS_TEST_CASE("If I call StartMatchmaking with name already in use, I get an error", EG_SESSIONS_STARTMATCHMAKING_TAG)
 {
 	DestroyCurrentServiceModule();
 
@@ -217,7 +216,7 @@ SESSIONS_TEST_CASE("If I call StartMatchmaking with name already in use, I get a
 	RunToCompletion();
 }
 
-SESSIONS_TEST_CASE("If I call StartMatchmaking with valid data, the operation will end with NotImplemented error for EOS and NULL", EG_SESSIONS_STARTMATCHMAKINGEOSNULL_TAG)
+SESSIONS_TEST_CASE("If I call StartMatchmaking with valid data, the operation will end with NotImplemented error for EOS and NULL", EG_SESSIONS_STARTMATCHMAKING_SKIP_TAG)
 {
 	FAccountId AccountId;
 
diff --git a/Engine/Source/Programs/Online/OnlineTestsCore/Source/Private/Tests/Sessions/UpdateSessionSettingsTests.cpp b/Engine/Source/Programs/Online/OnlineTestsCore/Source/Private/Tests/Sessions/UpdateSessionSettingsTests.cpp
--- a/Engine/Source/Programs/Online/OnlineTestsCore/Source/Private/Tests/Sessions/UpdateSessionSettingsTests.cpp
+++ b/Engine/Source/Programs/Online/OnlineTestsCore/Source/Private/Tests/Sessions/UpdateSessionSettingsTests.cpp
@@ -7,7 +7,6 @@
 
 #define SESSIONS_TAGS "[suite_sessions]"
 #define SESSIONS_UPDATESESSIONSETTINGS_TAG SESSIONS_TAGS "[updatesessionsettings]"
-#define SESSIONS_UPDATESESSIONSETTINGS_SKIP_TAG SESSIONS_TAGS "[updatesessionsettings][skip_UserNumToLogin]"
 #define UPDATESESSIONS_TEST_CASE(x, ...) ONLINE_TEST_CASE(x, SESSIONS_TAGS __VA_ARGS__)
 
 UPDATESESSIONS_TEST_CASE("If I call UpdateSessionSettings with an invalid account id, I get an error", SESSIONS_UPDATESESSIONSETTINGS_TAG)
@@ -127,7 +126,7 @@ UPDATESESSIONS_TEST_CASE("If I call UpdateSessionSettings with an unregistered s
 	RunToCompletion();
 }
 
-UPDATESESSIONS_TEST_CASE("If I call UpdateSessionSettings with valid data, the operation completes successfully", SESSIONS_UPDATESESSIONSETTINGS_SKIP_TAG)
+UPDATESESSIONS_TEST_CASE("If I call UpdateSessionSettings with valid data, the operation completes successfully", SESSIONS_UPDATESESSIONSETTINGS_TAG)
 {
 	DestroyCurrentServiceModule();
 
