Fix exclusive check causing test to be skipped when `--filename-as-tags` is used

diff --git a/Engine/Source/Programs/Online/OSSTestsCore/Helpers/OnlineSubsystemTestFixtures.h b/Engine/Source/Programs/Online/OSSTestsCore/Helpers/OnlineSubsystemTestFixtures.h
--- a/Engine/Source/Programs/Online/OSSTestsCore/Helpers/OnlineSubsystemTestFixtures.h
+++ b/Engine/Source/Programs/Online/OSSTestsCore/Helpers/OnlineSubsystemTestFixtures.h
@@ -23,6 +23,7 @@ public:
 	void DestroyCurrentOnlineSubsystemModule() const override 
 	{
 		ConfigInstanceName.Reset();
+		REQUIRE(!FOnlineSubsystemTestBaseFixture::GetSubsystem().IsEmpty());
 		FOnlineSubsystemTestBaseFixture::DestroyCurrentOnlineSubsystemModule();
 	}
 
diff --git a/Engine/Source/Programs/Online/OSSTestsCore/OnlineSubsystemCatchHelper.cpp b/Engine/Source/Programs/Online/OSSTestsCore/OnlineSubsystemCatchHelper.cpp
--- a/Engine/Source/Programs/Online/OSSTestsCore/OnlineSubsystemCatchHelper.cpp
+++ b/Engine/Source/Programs/Online/OSSTestsCore/OnlineSubsystemCatchHelper.cpp
@@ -41,55 +41,11 @@ TArray<TFunction<void()>>* GetGlobalInitalizers()
 	return &gInitalizersToCallInMain;
 }
 
-TArray<FOnlineSubsystemAutoReg::FApplicableServicesConfig> FOnlineSubsystemAutoReg::GetApplicableServices()
-{
-	static TArray<FApplicableServicesConfig> ServicesConfig =
-		[]()
-		{
-			TArray<FApplicableServicesConfig> ServicesConfigInit;
-			if (const TCHAR* CmdLine = FCommandLine::Get())	
-			{
-				FString Values;
-				TArray<FString> ServicesTags;
-				if (FParse::Value(CmdLine, TEXT("-Services="), Values, false))
-				{
-					Values.ParseIntoArray(ServicesTags, TEXT(","));
-				}
-
-				if (ServicesTags.IsEmpty())
-				{
-					GConfig->GetArray(TEXT("OnlineServicesTests"), TEXT("DefaultServices"), ServicesTags, GEngineIni);
-				}
-
-				for (const FString& ServicesTag : ServicesTags)
-				{
-					FString ConfigCategory = FString::Printf(TEXT("OnlineServicesTests %s"), *ServicesTag);
-					FApplicableServicesConfig Config;
-					Config.Tag = ServicesTag;
-
-					FString ServicesType;
-					GConfig->GetString(*ConfigCategory, TEXT("ServicesType"), ServicesType, GEngineIni);
-					GConfig->GetArray(*ConfigCategory, TEXT("ModulesToLoad"), Config.ModulesToLoad, GEngineIni);
-
-					LexFromString(Config.ServicesType, *ServicesType);
-					if (Config.ServicesType != UE::Online::EOnlineServices::None)
-					{
-						ServicesConfigInit.Add(MoveTemp(Config));
-					}
-				}
-			}
-
-			return ServicesConfigInit;
-		}();
-
-		return ServicesConfig;
-}
-
 TArray<FString> GetServiceModules()
 {
 	TArray<FString> Modules;
 
-	for (const FOnlineSubsystemAutoReg::FApplicableServicesConfig& Config : FOnlineSubsystemAutoReg::GetApplicableServices())
+	for (const FOnlineSubsystemFixtureInvoker::FApplicableServicesConfig& Config : FOnlineSubsystemFixtureInvoker::GetApplicableServices())
 	{
 		for (const FString& Module : Config.ModulesToLoad)
 		{
@@ -266,14 +222,89 @@ void FOnlineSubsystemTestBaseFixture::RunToCompletion() const
 	Driver.RunToCompletion();
 }
 
-TArray<FString> FOnlineSubsystemAutoReg::GetApplicableSubsystems()
+void FOnlineSubsystemFixtureInvoker::SetSubsystem(const FString& InStoredSubsystem) {
+	StoredSubsystem = InStoredSubsystem;
+}
+
+TArray<FOnlineSubsystemFixtureInvoker::FApplicableServicesConfig> FOnlineSubsystemFixtureInvoker::GetApplicableServices()
+{
+	static TArray<FApplicableServicesConfig> ServicesConfig =
+		[]()
+		{
+			TArray<FApplicableServicesConfig> ServicesConfigInit;
+			if (const TCHAR* CmdLine = FCommandLine::Get())
+			{
+				FString Values;
+				TArray<FString> ServicesTags;
+				if (FParse::Value(CmdLine, TEXT("-Services="), Values, false))
+				{
+					Values.ParseIntoArray(ServicesTags, TEXT(","));
+				}
+
+				if (ServicesTags.IsEmpty())
+				{
+					GConfig->GetArray(TEXT("OnlineServicesTests"), TEXT("DefaultServices"), ServicesTags, GEngineIni);
+				}
+
+				for (const FString& ServicesTag : ServicesTags)
+				{
+					FString ConfigCategory = FString::Printf(TEXT("OnlineServicesTests %s"), *ServicesTag);
+					FApplicableServicesConfig Config;
+					Config.Tag = ServicesTag;
+
+					FString ServicesType;
+					GConfig->GetString(*ConfigCategory, TEXT("ServicesType"), ServicesType, GEngineIni);
+					GConfig->GetArray(*ConfigCategory, TEXT("ModulesToLoad"), Config.ModulesToLoad, GEngineIni);
+
+					LexFromString(Config.ServicesType, *ServicesType);
+					if (Config.ServicesType != UE::Online::EOnlineServices::None)
+					{
+						ServicesConfigInit.Add(MoveTemp(Config));
+					}
+				}
+			}
+
+			return ServicesConfigInit;
+		}();
+
+	return ServicesConfig;
+}
+
+TArray<FOnlineSubsystemFixtureInvoker::FApplicableSubsystemConfig> FOnlineSubsystemFixtureInvoker::GetApplicableSubsystems()
 {
-	TArray<FString> Subsystems;
-	GConfig->GetArray(TEXT("OnlineSubsystemTests"), TEXT("Subsystems"), Subsystems, GEngineIni);
-	return Subsystems;
+	static TArray<FApplicableSubsystemConfig> SubsystemsConfig =
+		[]()
+		{
+			TArray<FApplicableSubsystemConfig> SubsystemsConfigInit;
+			if (const TCHAR* CmdLine = FCommandLine::Get())
+			{
+				FString Values;
+				TArray<FString> SubsystemsNames;
+				if (FParse::Value(CmdLine, TEXT("-Subsystems="), Values, false))
+				{
+					Values.ParseIntoArray(SubsystemsNames, TEXT(","));
+				}
+
+				if (SubsystemsNames.IsEmpty())
+				{
+					GConfig->GetArray(TEXT("OnlineSubsystemTests"), TEXT("Subsystems"), SubsystemsNames, GEngineIni);
+				}
+
+				for (const FString& SubsystemName : SubsystemsNames)
+				{
+					FApplicableSubsystemConfig Config;
+					Config.Name = SubsystemName;
+					SubsystemsConfigInit.Add(MoveTemp(Config));
+				}
+			}
+
+			return SubsystemsConfigInit;
+		}();
+
+	return SubsystemsConfig;
 }
 
-bool FOnlineSubsystemAutoReg::CheckAllTagsIsIn(const TArray<FString>& TestTags, const TArray<FString>& InputTags)
+bool FOnlineSubsystemFixtureInvoker::CheckAllTagsIsIn(const TArray<FString>& TestTags, const TArray<FString>& InputTags)
 {
 	if (InputTags.Num() == 0)
 	{
@@ -303,7 +334,7 @@ bool FOnlineSubsystemAutoReg::CheckAllTagsIsIn(const TArray<FString>& TestTags,
 	return bAllInputTagsInTestTags;
 }
 
-bool FOnlineSubsystemAutoReg::CheckAllTagsIsIn(const TArray<FString>& TestTags, const FString& RawTagString)
+bool FOnlineSubsystemFixtureInvoker::CheckAllTagsIsIn(const TArray<FString>& TestTags, const FString& RawTagString)
 {
 	TArray<FString> InputTags;
 	RawTagString.ParseIntoArray(InputTags, TEXT(","));
@@ -316,7 +347,7 @@ bool FOnlineSubsystemAutoReg::CheckAllTagsIsIn(const TArray<FString>& TestTags,
 	return CheckAllTagsIsIn(TestTags, InputTags);
 }
 
-FString FOnlineSubsystemAutoReg::GenerateTags(const FString& ServiceName, const FReportingSkippableTags& SkippableTags, const TCHAR* InTag)
+FString FOnlineSubsystemFixtureInvoker::GenerateTags(const FString& ServiceName, const FReportingSkippableTags& SkippableTags, const TCHAR* InTag)
 {
 	//Copy String here for ease-of-manipulation
 	FString RawInTag = InTag;
@@ -357,13 +388,10 @@ FString FOnlineSubsystemAutoReg::GenerateTags(const FString& ServiceName, const
 	return FString::Printf(TEXT("[%s] %s"), *ServiceName, *RawInTag);
 }
 
-bool FOnlineSubsystemAutoReg::ShouldDisableTest(const FString& ServiceName, const FReportingSkippableTags& SkippableTags, const TCHAR* InTag)
+FOnlineSubsystemFixtureInvoker::EDisableReason FOnlineSubsystemFixtureInvoker::ShouldDisableTest(const FString& ServiceName, const FReportingSkippableTags& SkippableTags, FString& InTag)
 {
-	//Copy String here for ease-of-manipulation
-	const FString RawInTag = InTag;
-
 	TArray<FString> TestTagsArray;
-	RawInTag.ParseIntoArray(TestTagsArray, TEXT("]"));
+	InTag.ParseIntoArray(TestTagsArray, TEXT("]"));
 	Algo::ForEach(TestTagsArray, [](FString& String)
 		{
 			String.TrimStartAndEndInline();
@@ -373,9 +401,9 @@ bool FOnlineSubsystemAutoReg::ShouldDisableTest(const FString& ServiceName, cons
 
 	// If we contain [!<service>] it means we shouldn't run this
 	// test against this service.
-	if (RawInTag.Contains("!" + ServiceName))
+	if (InTag.Contains("!" + ServiceName))
 	{
-		return true;
+		return EDisableReason::AgainstService;
 	}
 
 	// Check for exclusive runs
@@ -384,14 +412,21 @@ bool FOnlineSubsystemAutoReg::ShouldDisableTest(const FString& ServiceName, cons
 		const FString& ServiceTag = Config.Tag;
 		if (ServiceName.Equals(ServiceTag, ESearchCase::IgnoreCase))
 		{
+			if (InTag.Contains("." + ServiceTag))
+			{
+				FString SubstringToRemove = TEXT(".");
+				FString Replacement = TEXT("");
+				InTag.ReplaceInline(*SubstringToRemove, *Replacement);
+			}
+			
 			continue;
 		}
 
 		// If we contain [.NULL] and we're running with [EOS] we shouldn't
 		// generate a test for [EOS] here.
-		if (RawInTag.Contains("." + ServiceTag))
+		if (InTag.Contains("." + ServiceTag))
 		{
-			return true;
+			return EDisableReason::ExclusiveService;
 		}
 	}
 
@@ -401,46 +436,95 @@ bool FOnlineSubsystemAutoReg::ShouldDisableTest(const FString& ServiceName, cons
 	{
 		if (CheckAllTagsIsIn(TestTagsArray, DisableTag))
 		{
-			return true;
+			return EDisableReason::DisableTagPresence;
 		}
 	}
 
 	// We should run the test!
-	return false;
+	return EDisableReason::Success;
+}
+
+FOnlineSubsystemFixtureInvoker::EDisableReason FOnlineSubsystemFixtureInvoker::ShouldSkipTest(const FString& ServiceName, const FReportingSkippableTags& SkippableTags, FString& InTag) {
+
+	// If we have tags present indicating we should exit the test
+	if (EDisableReason Reason = ShouldDisableTest(ServiceName, SkippableTags, InTag); Reason != EDisableReason::Success)
+	{
+		return Reason;
+	}
+
+	return EDisableReason::Success;
+}
+
+bool FOnlineSubsystemFixtureInvoker::IsRunningTestSkipOnTags(const FString& ServiceName, const FReportingSkippableTags& SkippableTags, FString& InTag)
+{
+	bool bIsShouldSkip = true;
+	EDisableReason Reason = ShouldSkipTest(ServiceName, SkippableTags, InTag);
+
+	switch (Reason)
+	{
+		case EDisableReason::Success:
+		{
+			bIsShouldSkip = false;
+			break;
+		}
+		case EDisableReason::AgainstService:
+		{
+			UE_LOG(LogOSSTests, Verbose, TEXT("Test skipped due to run against this service."));
+			break;
+		}
+		case EDisableReason::ExclusiveService:
+		{
+			UE_LOG(LogOSSTests, Verbose, TEXT("Test skipped due to exclusive service run."));
+			break;
+		}
+		case EDisableReason::DisableTagPresence:
+		{
+			UE_LOG(LogOSSTests, Verbose, TEXT("Test skipped due to disable tag presence."));
+			break;
+		}
+		default:
+		{
+			UE_LOG(LogOSSTests, Error, TEXT("Test skipped due to unknown reason!"));
+			break;
+		}
+	}
+			
+	return bIsShouldSkip;
 }
 
 // This code is kept identical to Catch internals so that there is as little deviation from OSS_TESTS and Online_OSS_TESTS as possible
 FOnlineSubsystemAutoReg::FOnlineSubsystemAutoReg(Catch::Detail::unique_ptr<FOnlineSubsystemFixtureInvoker> TestInvoker, Catch::SourceLineInfo LineInfo, const char* Name, const char* Tags, const char* AddlOnlineInfo)
 {
 	auto TestInvokerPtr = TestInvoker.release();
-
 	auto GlobalInitalizersPtr = GetGlobalInitalizers();
 	ensure(GlobalInitalizersPtr);
 	GlobalInitalizersPtr->Add([TestInvokerPtr = MoveTemp(TestInvokerPtr), LineInfo, Name, Tags, this]() mutable -> void
 		{
-			for (const FString& Subsystem : GetApplicableSubsystems())
+			for (const FOnlineSubsystemFixtureInvoker::FApplicableSubsystemConfig& Config : FOnlineSubsystemFixtureInvoker::GetApplicableSubsystems())
 			{
-				FString ReportingCategory = FString::Printf(TEXT("TestReporting %s"), *Subsystem);
-				FReportingSkippableTags SkippableTags;
+				const FString SubsystemName = Config.Name;
+				FString ReportingCategory = FString::Printf(TEXT("TestReporting %s"), *SubsystemName);
+				FOnlineSubsystemFixtureInvoker::FReportingSkippableTags SkippableTags;
 				GConfig->GetArray(*ReportingCategory, TEXT("MayFailTestTags"), SkippableTags.MayFailTags, GEngineIni);
 				GConfig->GetArray(*ReportingCategory, TEXT("ShouldFailTestTags"), SkippableTags.ShouldFailTags, GEngineIni);
 				GConfig->GetArray(*ReportingCategory, TEXT("DisableTestTags"), SkippableTags.DisableTestTags, GEngineIni);
 
-				auto NewName = StringCast<ANSICHAR>(*FString::Printf(TEXT("[%s] %s"), *Subsystem, ANSI_TO_TCHAR(Name)));
-				auto NewTags = StringCast<ANSICHAR>(*GenerateTags(Subsystem, SkippableTags, ANSI_TO_TCHAR(Tags)));
+				auto NewName = StringCast<ANSICHAR>(*FString::Printf(TEXT("[%s] %s"), *SubsystemName, ANSI_TO_TCHAR(Name)));
+				auto GeneratedTags = StringCast<ANSICHAR>(*FOnlineSubsystemFixtureInvoker::GenerateTags(SubsystemName, SkippableTags, ANSI_TO_TCHAR(Tags)));
+				FString NewTags = GeneratedTags.Get();
 
-				// If we have tags present indicating we should not enable the test at all
-				if (ShouldDisableTest(Subsystem, SkippableTags, ANSI_TO_TCHAR(NewTags.Get())))
+				if (!FOnlineSubsystemFixtureInvoker::IsRunningTestSkipOnTags(SubsystemName, SkippableTags, NewTags))
+				{
+					TestInvokerPtr->SetSubsystem(Config.Name);
+				}
+				else
 				{
 					continue;
 				}
 
-				TestInvokerPtr->SetSubsystem(Subsystem);
-
-				// This code is lifted from Catch internals to register a test
 				Catch::getMutableRegistryHub().registerTest(Catch::makeTestCaseInfo(
 					std::string(Catch::StringRef()),  // Used for testing a static method instead of a function- not needed since we're passing an ITestInvoker macro
-					Catch::NameAndTags{ NewName.Get(), NewTags.Get() },
+					Catch::NameAndTags{ NewName.Get(), TCHAR_TO_ANSI(*NewTags) },
 					LineInfo),
 					Catch::Detail::unique_ptr(TestInvokerPtr) // This is taking the ITestInvoker macro and will call invoke() to run the test
 				);
diff --git a/Engine/Source/Programs/Online/OSSTestsCore/OnlineSubsystemCatchHelper.h b/Engine/Source/Programs/Online/OSSTestsCore/OnlineSubsystemCatchHelper.h
--- a/Engine/Source/Programs/Online/OSSTestsCore/OnlineSubsystemCatchHelper.h
+++ b/Engine/Source/Programs/Online/OSSTestsCore/OnlineSubsystemCatchHelper.h
@@ -16,11 +16,61 @@
 class FOnlineSubsystemFixtureInvoker : public Catch::ITestInvoker
 {
 public:
+
+	struct FReportingSkippableTags
+	{
+		TArray<FString> MayFailTags;
+		TArray<FString> ShouldFailTags;
+		TArray<FString> DisableTestTags;
+	};
+
+	struct FApplicableServicesConfig
+	{
+		FString Tag;
+		UE::Online::EOnlineServices ServicesType;
+		TArray<FString> ModulesToLoad;
+	};
+
+	struct FApplicableSubsystemConfig
+	{
+		FString Name;
+	};
+
+	enum class EDisableReason { Success, AgainstService, ExclusiveService, DisableTagPresence };
+
 	virtual void invoke() const = 0;
 
-	void SetSubsystem(const FString& InStoredSubsystem) {
-		StoredSubsystem = InStoredSubsystem;
-	}
+	void SetSubsystem(const FString& InStoredSubsystem);
 
+	// TODO: This should poll some info from the tags and generate the list based on that
+	static TArray<FApplicableServicesConfig> GetApplicableServices();
+	static TArray<FApplicableSubsystemConfig> GetApplicableSubsystems();
+
+	/*
+	* Helper function that calls CheckAllTagsIsIn(const TArray<FString>&, const TArray<FString>&);
+	*
+	* @param  TestTags			The array of test tags we wish to test against.
+	* @param  RawTagString		Comma seperated string of elemnets we wish to convert to an array.
+	*
+	* @return true if all elements of RawTagString prased as an comma sperated array is in TestTags.
+	*/
+	static bool CheckAllTagsIsIn(const TArray<FString>& TestTags, const FString& RawTagString);
+
+	/**
+	 * Checks if every element of InputTags is in TestTags.
+	 *
+	 * @param  TestTags       The array of test tags we wish to test against.
+	 * @param  InputTags	  The array of tags we wish to see if all elements of are present in TestTags.
+	 *
+	 * @return  true if all elements of InputTags is in TestTags.
+	 */
+	static bool CheckAllTagsIsIn(const TArray<FString>& TestTags, const TArray<FString>& InputTags);
+
+	static FString GenerateTags(const FString& ServiceName, const FReportingSkippableTags& SkippableTags, const TCHAR* InTag);
+
+	static EDisableReason ShouldDisableTest(const FString& ServiceName, const FReportingSkippableTags& SkippableTags, FString& InTag);
+	static EDisableReason ShouldSkipTest(const FString& ServiceName, const FReportingSkippableTags& SkippableTags, FString& InTag);
+	static bool IsRunningTestSkipOnTags(const FString& ServiceName, const FReportingSkippableTags& SkippableTags, FString& InTag);
+
 protected:
 	FString StoredSubsystem;
@@ -30,7 +80,7 @@
 class TestInvokerFixture : public FOnlineSubsystemFixtureInvoker
 {
 public:
-	constexpr TestInvokerFixture(void (C::* InTestAsMethod)() const) noexcept :
+	constexpr TestInvokerFixture(void (C::* InTestAsMethod)()) noexcept :
 		TestAsMethod(InTestAsMethod) {}
 
 	void PrepareTestCase() const {
@@ -50,12 +100,12 @@
 	}
 
 private:
-	void (C::* TestAsMethod)() const;
+	void (C::* TestAsMethod)();
 	mutable Catch::Detail::unique_ptr<C> Fixture = nullptr;
 };
 
 template<typename C>
-Catch::Detail::unique_ptr<FOnlineSubsystemFixtureInvoker> MakeTestInvokerFixture(void (C::* TestAsMethod)() const) {
+Catch::Detail::unique_ptr<FOnlineSubsystemFixtureInvoker> MakeTestInvokerFixture(void (C::* TestAsMethod)()) {
 	return Catch::Detail::make_unique<TestInvokerFixture<C>>(TestAsMethod);
 }
 
@@ -110,54 +160,9 @@
 
 class FOnlineSubsystemAutoReg
 {
-private:
-	// TODO: This should poll some info from the tags and generate the list based on that
-	TArray<FString> GetApplicableSubsystems();
-
 public:
-	struct FReportingSkippableTags
-	{
-		TArray<FString> MayFailTags;
-		TArray<FString> ShouldFailTags;
-		TArray<FString> DisableTestTags;
-	};
-		
-	struct FApplicableServicesConfig
-	{
-		FString Tag;
-		UE::Online::EOnlineServices ServicesType;
-		TArray<FString> ModulesToLoad;
-	};
-
-	static TArray<FApplicableServicesConfig> GetApplicableServices();
-
-	/*
-	* Helper function that calls CheckAllTagsIsIn(const TArray<FString>&, const TArray<FString>&);
-	*
-	* @param  TestTags			The array of test tags we wish to test against.
-	* @param  RawTagString		Comma seperated string of elemnets we wish to convert to an array.
-	*
-	* @return true if all elements of RawTagString prased as an comma sperated array is in TestTags.
-	*/
-	static bool CheckAllTagsIsIn(const TArray<FString>& TestTags, const FString& RawTagString);
-
-
-	/**
-	 * Checks if every element of InputTags is in TestTags.
-	 *
-	 * @param  TestTags       The array of test tags we wish to test against.
-	 * @param  InputTags	  The array of tags we wish to see if all elements of are present in TestTags.
-	 *
-	 * @return  true if all elements of InputTags is in TestTags.
-	 */
-	static bool CheckAllTagsIsIn(const TArray<FString>& TestTags, const TArray<FString>& InputTags);
-
-	static FString GenerateTags(const FString& ServiceName, const FReportingSkippableTags& SkippableTags, const TCHAR* InTag);
-	static bool ShouldDisableTest(const FString& ServiceName, const FReportingSkippableTags& SkippableTags, const TCHAR* InTag);
-
 	// This code is kept identical to Catch internals so that there is as little deviation from OSS_TESTS and Online_OSS_TESTS as possible
 	FOnlineSubsystemAutoReg(Catch::Detail::unique_ptr<FOnlineSubsystemFixtureInvoker> TestInvoker, Catch::SourceLineInfo LineInfo, const char* Name, const char* Tags, const char* AddlOnlineInfo);
-
 };
 
 TArray<TFunction<void()>>* GetGlobalInitalizers();
@@ -167,11 +172,14 @@
 class PREPROCESSOR_JOIN(OnlineSubsystemTest_,RegName) : public ClassName {\
 public:\
 	virtual void invoke() const override;\
-	void test() const;\
+	void test();\
 };\
 	FOnlineSubsystemAutoReg RegName( ::MakeTestInvokerFixture( PREPROCESSOR_JOIN(&OnlineSubsystemTest_,RegName)::test ), CATCH_INTERNAL_LINEINFO, Name, Tags, "");\
 }\
-void PREPROCESSOR_JOIN(OnlineSubsystemTest_,RegName)::test() const { invoke(); }\
+void PREPROCESSOR_JOIN(OnlineSubsystemTest_,RegName)::test()\
+{\
+	invoke();\
+}\
 void PREPROCESSOR_JOIN(OnlineSubsystemTest_,RegName)::invoke() const\
 
 #define ONLINESUBSYSTEM_TEST_CASE_FIXTURE(ClassName, Name, Tags, ...) \
diff --git a/Engine/Source/Programs/Online/OSSTestsCore/Tests/OSSCatchHelpers/OSSCatchHelpers.cpp b/Engine/Source/Programs/Online/OSSTestsCore/Tests/OSSCatchHelpers/OSSCatchHelpers.cpp
--- a/Engine/Source/Programs/Online/OSSTestsCore/Tests/OSSCatchHelpers/OSSCatchHelpers.cpp
+++ b/Engine/Source/Programs/Online/OSSTestsCore/Tests/OSSCatchHelpers/OSSCatchHelpers.cpp
@@ -11,100 +11,100 @@
 SELFTEST_TEST_CASE("GenerateTags append MayFailTags case", GENERATETAGS_TAG)
 {
 	FString TestTags = "[foo][bar]";
-	FOnlineSubsystemAutoReg::FReportingSkippableTags SkipTags;
+	FOnlineSubsystemFixtureInvoker::FReportingSkippableTags SkipTags;
 	SkipTags.MayFailTags = { "[foo]" };
-	FString OutTags = FOnlineSubsystemAutoReg::GenerateTags("TestService", SkipTags, *TestTags);
+	FString OutTags = FOnlineSubsystemFixtureInvoker::GenerateTags("TestService", SkipTags, *TestTags);
 	CHECK(OutTags == TEXT("[TestService] [foo][bar][!mayfail]"));
 }
 
 SELFTEST_TEST_CASE("GenerateTags append multiple match MayFailTags case", GENERATETAGS_TAG)
 {
 	FString TestTags = "[foo][bar]";
-	FOnlineSubsystemAutoReg::FReportingSkippableTags SkipTags;
+	FOnlineSubsystemFixtureInvoker::FReportingSkippableTags SkipTags;
 	SkipTags.MayFailTags = { "[bar]", "[foo]" };
-	FString OutTags = FOnlineSubsystemAutoReg::GenerateTags("TestService", SkipTags, *TestTags);
+	FString OutTags = FOnlineSubsystemFixtureInvoker::GenerateTags("TestService", SkipTags, *TestTags);
 	CHECK(OutTags == TEXT("[TestService] [foo][bar][!mayfail]"));
 }
 
 SELFTEST_TEST_CASE("GenerateTags append by last match MayFailTags case", GENERATETAGS_TAG)
 {
 	FString TestTags = "[foo][bar]";
-	FOnlineSubsystemAutoReg::FReportingSkippableTags SkipTags;
+	FOnlineSubsystemFixtureInvoker::FReportingSkippableTags SkipTags;
 	SkipTags.MayFailTags = { "[wiz]", "[foo]" };
-	FString OutTags = FOnlineSubsystemAutoReg::GenerateTags("TestService", SkipTags, *TestTags);
+	FString OutTags = FOnlineSubsystemFixtureInvoker::GenerateTags("TestService", SkipTags, *TestTags);
 	CHECK(OutTags == TEXT("[TestService] [foo][bar][!mayfail]"));
 }
 
 SELFTEST_TEST_CASE("GenerateTags append by last match mutli-tag case", GENERATETAGS_TAG)
 {
 	FString TestTags = "[foo][bar]";
-	FOnlineSubsystemAutoReg::FReportingSkippableTags SkipTags;
+	FOnlineSubsystemFixtureInvoker::FReportingSkippableTags SkipTags;
 	SkipTags.MayFailTags = { "[foo],bar" };
-	FString OutTags = FOnlineSubsystemAutoReg::GenerateTags("TestService", SkipTags, *TestTags);
+	FString OutTags = FOnlineSubsystemFixtureInvoker::GenerateTags("TestService", SkipTags, *TestTags);
 	CHECK(OutTags == TEXT("[TestService] [foo][bar][!mayfail]"));
 }
 
 SELFTEST_TEST_CASE("GenerateTags append by last match mutli-tag no match case", GENERATETAGS_TAG)
 {
 	FString TestTags = "[foo][bar]";
-	FOnlineSubsystemAutoReg::FReportingSkippableTags SkipTags;
+	FOnlineSubsystemFixtureInvoker::FReportingSkippableTags SkipTags;
 	SkipTags.MayFailTags = { "[foo],[wiz]" };
-	FString OutTags = FOnlineSubsystemAutoReg::GenerateTags("TestService", SkipTags, *TestTags);
+	FString OutTags = FOnlineSubsystemFixtureInvoker::GenerateTags("TestService", SkipTags, *TestTags);
 	CHECK(OutTags == TEXT("[TestService] [foo][bar]"));
 }
 
 SELFTEST_TEST_CASE("GenerateTags don't append MayFailTags no match case", GENERATETAGS_TAG)
 {
-	FOnlineSubsystemAutoReg::FReportingSkippableTags SkipTags;
+	FOnlineSubsystemFixtureInvoker::FReportingSkippableTags SkipTags;
 	SkipTags.MayFailTags = { "[wiz]" };
 	FString TestTags = "[foo][bar]";
-	FString OutTags = FOnlineSubsystemAutoReg::GenerateTags("TestService", SkipTags, *TestTags);
+	FString OutTags = FOnlineSubsystemFixtureInvoker::GenerateTags("TestService", SkipTags, *TestTags);
 	CHECK(OutTags == TEXT("[TestService] [foo][bar]"));
 }
 
 SELFTEST_TEST_CASE("GenerateTags append ShouldFail case", GENERATETAGS_TAG)
 {
 	FString TestTags = "[foo][bar]";
-	FOnlineSubsystemAutoReg::FReportingSkippableTags SkipTags;
+	FOnlineSubsystemFixtureInvoker::FReportingSkippableTags SkipTags;
 	SkipTags.ShouldFailTags = { "[foo]" };
-	FString OutTags = FOnlineSubsystemAutoReg::GenerateTags("TestService", SkipTags, *TestTags);
+	FString OutTags = FOnlineSubsystemFixtureInvoker::GenerateTags("TestService", SkipTags, *TestTags);
 	CHECK(OutTags == TEXT("[TestService] [foo][bar][!shouldfail]"));
 }
 
 SELFTEST_TEST_CASE("GenerateTags append multiple match ShouldFail case", GENERATETAGS_TAG)
 {
 	FString TestTags = "[foo][bar]";
-	FOnlineSubsystemAutoReg::FReportingSkippableTags SkipTags;
+	FOnlineSubsystemFixtureInvoker::FReportingSkippableTags SkipTags;
 	SkipTags.ShouldFailTags = { "[bar]", "[foo]" };
-	FString OutTags = FOnlineSubsystemAutoReg::GenerateTags("TestService", SkipTags, *TestTags);
+	FString OutTags = FOnlineSubsystemFixtureInvoker::GenerateTags("TestService", SkipTags, *TestTags);
 	CHECK(OutTags == TEXT("[TestService] [foo][bar][!shouldfail]"));
 }
 
 SELFTEST_TEST_CASE("GenerateTags append by last match ShouldFail case", GENERATETAGS_TAG)
 {
 	FString TestTags = "[foo][bar]";
-	FOnlineSubsystemAutoReg::FReportingSkippableTags SkipTags;
+	FOnlineSubsystemFixtureInvoker::FReportingSkippableTags SkipTags;
 	SkipTags.ShouldFailTags = { "[wiz]", "[foo]" };
-	FString OutTags = FOnlineSubsystemAutoReg::GenerateTags("TestService", SkipTags, *TestTags);
+	FString OutTags = FOnlineSubsystemFixtureInvoker::GenerateTags("TestService", SkipTags, *TestTags);
 	CHECK(OutTags == TEXT("[TestService] [foo][bar][!shouldfail]"));
 }
 
 SELFTEST_TEST_CASE("GenerateTags don't append ShouldFail no match case", GENERATETAGS_TAG)
 {
 	FString TestTags = "[foo][bar]";
-	FOnlineSubsystemAutoReg::FReportingSkippableTags SkipTags;
+	FOnlineSubsystemFixtureInvoker::FReportingSkippableTags SkipTags;
 	SkipTags.ShouldFailTags = { "[wiz]" };
-	FString OutTags = FOnlineSubsystemAutoReg::GenerateTags("TestService", SkipTags, *TestTags);
+	FString OutTags = FOnlineSubsystemFixtureInvoker::GenerateTags("TestService", SkipTags, *TestTags);
 	CHECK(OutTags == TEXT("[TestService] [foo][bar]"));
 }
 
 SELFTEST_TEST_CASE("GenerateTags appends MayFail and ShouldFail Case", GENERATETAGS_TAG)
 {
-	FOnlineSubsystemAutoReg::FReportingSkippableTags SkipTags;
+	FOnlineSubsystemFixtureInvoker::FReportingSkippableTags SkipTags;
 	SkipTags.MayFailTags = { "[foo]" };
 	SkipTags.ShouldFailTags = { "[bar]" };
 	FString TestTags = "[foo][bar]";
-	FString OutTags = FOnlineSubsystemAutoReg::GenerateTags("TestService", SkipTags, *TestTags);
+	FString OutTags = FOnlineSubsystemFixtureInvoker::GenerateTags("TestService", SkipTags, *TestTags);
 	CHECK(OutTags == TEXT("[TestService] [foo][bar][!mayfail][!shouldfail]"));
 }
 
@@ -112,65 +112,65 @@ SELFTEST_TEST_CASE("GenerateTags appends MayFail and ShouldFail Case", GENERATET
 SELFTEST_TEST_CASE("GenerateTags append by last match mutli-tag ShouldFail case", GENERATETAGS_TAG)
 {
 	FString TestTags = "[foo][bar]";
-	FOnlineSubsystemAutoReg::FReportingSkippableTags SkipTags;
+	FOnlineSubsystemFixtureInvoker::FReportingSkippableTags SkipTags;
 	SkipTags.ShouldFailTags = { "[foo],bar" };
-	FString OutTags = FOnlineSubsystemAutoReg::GenerateTags("TestService", SkipTags, *TestTags);
+	FString OutTags = FOnlineSubsystemFixtureInvoker::GenerateTags("TestService", SkipTags, *TestTags);
 	CHECK(OutTags == TEXT("[TestService] [foo][bar][!shouldfail]"));
 }
 
 SELFTEST_TEST_CASE("GenerateTags append by last match mutli-tag no match ShouldFail case", GENERATETAGS_TAG)
 {
 	FString TestTags = "[foo][bar]";
-	FOnlineSubsystemAutoReg::FReportingSkippableTags SkipTags;
+	FOnlineSubsystemFixtureInvoker::FReportingSkippableTags SkipTags;
 	SkipTags.ShouldFailTags = { "[foo],[wiz]" };
-	FString OutTags = FOnlineSubsystemAutoReg::GenerateTags("TestService", SkipTags, *TestTags);
+	FString OutTags = FOnlineSubsystemFixtureInvoker::GenerateTags("TestService", SkipTags, *TestTags);
 	CHECK(OutTags == TEXT("[TestService] [foo][bar]"));
 }
 
 SELFTEST_TEST_CASE("GenerateTags appends no tags case", GENERATETAGS_TAG)
 {
-	FOnlineSubsystemAutoReg::FReportingSkippableTags SkipTags;
+	FOnlineSubsystemFixtureInvoker::FReportingSkippableTags SkipTags;
 	SkipTags.MayFailTags = { "[wiz]" };
 	SkipTags.ShouldFailTags = {  };
 	FString TestTags = "[foo][bar]";
-	FString OutTags = FOnlineSubsystemAutoReg::GenerateTags("TestService", SkipTags, *TestTags);
+	FString OutTags = FOnlineSubsystemFixtureInvoker::GenerateTags("TestService", SkipTags, *TestTags);
 	CHECK(OutTags == TEXT("[TestService] [foo][bar]"));
 }
 
 SELFTEST_TEST_CASE("ShouldDisableTest returns true on single-tag config", SHOULDDISABLETEST_TAG)
 {
-	FOnlineSubsystemAutoReg::FReportingSkippableTags SkipTags;
+	FOnlineSubsystemFixtureInvoker::FReportingSkippableTags SkipTags;
 	SkipTags.DisableTestTags = { "[foo]" };
 	FString TestTags = "[foo][bar]";
-	CHECK(FOnlineSubsystemAutoReg::ShouldDisableTest("TestService", SkipTags, *TestTags) == true);
+	CHECK(FOnlineSubsystemFixtureInvoker::ShouldDisableTest("TestService", SkipTags, TestTags) == FOnlineSubsystemFixtureInvoker::EDisableReason::DisableTagPresence);
 }
 
 SELFTEST_TEST_CASE("ShouldDisableTest returns true on multi-tag config", SHOULDDISABLETEST_TAG)
 {
-	FOnlineSubsystemAutoReg::FReportingSkippableTags SkipTags;
+	FOnlineSubsystemFixtureInvoker::FReportingSkippableTags SkipTags;
 	SkipTags.DisableTestTags = { "[foo],bar" };
 	FString TestTags = "[foo][bar]";
-	CHECK(FOnlineSubsystemAutoReg::ShouldDisableTest("TestService", SkipTags, *TestTags) == true);
+	CHECK(FOnlineSubsystemFixtureInvoker::ShouldDisableTest("TestService", SkipTags, TestTags) == FOnlineSubsystemFixtureInvoker::EDisableReason::DisableTagPresence);
 }
 
 SELFTEST_TEST_CASE("ShouldDisableTest returns true on !<service>", SHOULDDISABLETEST_TAG)
 {
 	FString TestTags = "[foo][bar][!TestService]";
-	CHECK(FOnlineSubsystemAutoReg::ShouldDisableTest("TestService", {}, *TestTags) == true);
+	CHECK(FOnlineSubsystemFixtureInvoker::ShouldDisableTest("TestService", {}, TestTags) == FOnlineSubsystemFixtureInvoker::EDisableReason::AgainstService);
 }
 
 SELFTEST_TEST_CASE("ShouldDisableTest returns false with no tags and no config skips", SHOULDDISABLETEST_TAG)
 {
 	FString TestTags = "[foo][bar]";
-	CHECK(FOnlineSubsystemAutoReg::ShouldDisableTest("TestService", {}, *TestTags) == false);
+	CHECK(FOnlineSubsystemFixtureInvoker::ShouldDisableTest("TestService", {}, TestTags) == FOnlineSubsystemFixtureInvoker::EDisableReason::Success);
 }
 
 SELFTEST_TEST_CASE("ShouldDisableTest returns false with no matching no-tags and no matching multi-tag config skips", SHOULDDISABLETEST_TAG)
 {
-	FOnlineSubsystemAutoReg::FReportingSkippableTags SkipTags;
+	FOnlineSubsystemFixtureInvoker::FReportingSkippableTags SkipTags;
 	SkipTags.DisableTestTags = { "[foo],wiz" };
 	FString TestTags = "[foo][bar]";
-	CHECK(FOnlineSubsystemAutoReg::ShouldDisableTest("TestService", SkipTags, *TestTags) == false);
+	CHECK(FOnlineSubsystemFixtureInvoker::ShouldDisableTest("TestService", SkipTags, TestTags) == FOnlineSubsystemFixtureInvoker::EDisableReason::Success);
 }
 
 SELFTEST_TEST_CASE("CheckAllTagsIsIn(TArray, FString) true cases", CHECKALLTAGSISIN_TAG)
@@ -179,30 +179,30 @@ SELFTEST_TEST_CASE("CheckAllTagsIsIn(TArray, FString) true cases", CHECKALLTAGSI
 	CAPTURE(TestTags);
 
 	// Truthy Cases
-	CHECK(FOnlineSubsystemAutoReg::CheckAllTagsIsIn(TestTags, "bob, alice") == true);
-	CHECK(FOnlineSubsystemAutoReg::CheckAllTagsIsIn(TestTags, "bob,alice") == true);
-	CHECK(FOnlineSubsystemAutoReg::CheckAllTagsIsIn(TestTags, " bob,alice ") == true);
-	CHECK(FOnlineSubsystemAutoReg::CheckAllTagsIsIn(TestTags, "foo") == true);
-	CHECK(FOnlineSubsystemAutoReg::CheckAllTagsIsIn(TestTags, ",foo") == true);
-	CHECK(FOnlineSubsystemAutoReg::CheckAllTagsIsIn(TestTags, "bob,alice,foo") == true);
+	CHECK(FOnlineSubsystemFixtureInvoker::CheckAllTagsIsIn(TestTags, "bob, alice") == true);
+	CHECK(FOnlineSubsystemFixtureInvoker::CheckAllTagsIsIn(TestTags, "bob,alice") == true);
+	CHECK(FOnlineSubsystemFixtureInvoker::CheckAllTagsIsIn(TestTags, " bob,alice ") == true);
+	CHECK(FOnlineSubsystemFixtureInvoker::CheckAllTagsIsIn(TestTags, "foo") == true);
+	CHECK(FOnlineSubsystemFixtureInvoker::CheckAllTagsIsIn(TestTags, ",foo") == true);
+	CHECK(FOnlineSubsystemFixtureInvoker::CheckAllTagsIsIn(TestTags, "bob,alice,foo") == true);
 
 	//Bracket Parsing
-	CHECK(FOnlineSubsystemAutoReg::CheckAllTagsIsIn(TestTags, "[bob],[alice],[foo]") == true);
-	CHECK(FOnlineSubsystemAutoReg::CheckAllTagsIsIn(TestTags, "[bob], [alice,foo]") == true);
-	CHECK(FOnlineSubsystemAutoReg::CheckAllTagsIsIn(TestTags, "bob],  alice],  [foo]  ,") == true);
-	CHECK(FOnlineSubsystemAutoReg::CheckAllTagsIsIn(TestTags, ",[foo]") == true);
-	CHECK(FOnlineSubsystemAutoReg::CheckAllTagsIsIn(TestTags, ",foo]") == true);
-	CHECK(FOnlineSubsystemAutoReg::CheckAllTagsIsIn(TestTags, "[wiz]") == false);
+	CHECK(FOnlineSubsystemFixtureInvoker::CheckAllTagsIsIn(TestTags, "[bob],[alice],[foo]") == true);
+	CHECK(FOnlineSubsystemFixtureInvoker::CheckAllTagsIsIn(TestTags, "[bob], [alice,foo]") == true);
+	CHECK(FOnlineSubsystemFixtureInvoker::CheckAllTagsIsIn(TestTags, "bob],  alice],  [foo]  ,") == true);
+	CHECK(FOnlineSubsystemFixtureInvoker::CheckAllTagsIsIn(TestTags, ",[foo]") == true);
+	CHECK(FOnlineSubsystemFixtureInvoker::CheckAllTagsIsIn(TestTags, ",foo]") == true);
+	CHECK(FOnlineSubsystemFixtureInvoker::CheckAllTagsIsIn(TestTags, "[wiz]") == false);
 
 	// Negative Cases
-	CHECK(FOnlineSubsystemAutoReg::CheckAllTagsIsIn(TestTags, "bob,alice,foo,wiz") == false);
-	CHECK(FOnlineSubsystemAutoReg::CheckAllTagsIsIn(TestTags, "bob,wiz") == false);
-	CHECK(FOnlineSubsystemAutoReg::CheckAllTagsIsIn(TestTags, ",wiz") == false);
+	CHECK(FOnlineSubsystemFixtureInvoker::CheckAllTagsIsIn(TestTags, "bob,alice,foo,wiz") == false);
+	CHECK(FOnlineSubsystemFixtureInvoker::CheckAllTagsIsIn(TestTags, "bob,wiz") == false);
+	CHECK(FOnlineSubsystemFixtureInvoker::CheckAllTagsIsIn(TestTags, ",wiz") == false);
 
 	// Bound Checks
-	CHECK(FOnlineSubsystemAutoReg::CheckAllTagsIsIn(TestTags, ",") == false);
-	CHECK(FOnlineSubsystemAutoReg::CheckAllTagsIsIn(TestTags, "") == false);
-	CHECK(FOnlineSubsystemAutoReg::CheckAllTagsIsIn({}, "") == false);
+	CHECK(FOnlineSubsystemFixtureInvoker::CheckAllTagsIsIn(TestTags, ",") == false);
+	CHECK(FOnlineSubsystemFixtureInvoker::CheckAllTagsIsIn(TestTags, "") == false);
+	CHECK(FOnlineSubsystemFixtureInvoker::CheckAllTagsIsIn({}, "") == false);
 }
 
 
@@ -212,18 +212,18 @@ SELFTEST_TEST_CASE("CheckAllTagsIsIn(TArray, TArray) true cases", CHECKALLTAGSIS
 	CAPTURE(TestTags);
 
 	// Truthy Cases
-	CHECK(FOnlineSubsystemAutoReg::CheckAllTagsIsIn(TestTags, { "bob", "alice" }) == true);
-	CHECK(FOnlineSubsystemAutoReg::CheckAllTagsIsIn(TestTags, { "bob", "alice", "foo" }) == true);
-	CHECK(FOnlineSubsystemAutoReg::CheckAllTagsIsIn(TestTags, TArray<FString>({ "foo" })) == true);
+	CHECK(FOnlineSubsystemFixtureInvoker::CheckAllTagsIsIn(TestTags, { "bob", "alice" }) == true);
+	CHECK(FOnlineSubsystemFixtureInvoker::CheckAllTagsIsIn(TestTags, { "bob", "alice", "foo" }) == true);
+	CHECK(FOnlineSubsystemFixtureInvoker::CheckAllTagsIsIn(TestTags, TArray<FString>({ "foo" })) == true);
 
 	// Negative Cases
-	CHECK(FOnlineSubsystemAutoReg::CheckAllTagsIsIn(TestTags, TArray<FString>({ "wiz" })) == false);
-	CHECK(FOnlineSubsystemAutoReg::CheckAllTagsIsIn(TestTags, { "bob", "alice", "foo", "wiz" }) == false);
-	CHECK(FOnlineSubsystemAutoReg::CheckAllTagsIsIn(TestTags, { "bob", "alice", "wiz" }) == false);
+	CHECK(FOnlineSubsystemFixtureInvoker::CheckAllTagsIsIn(TestTags, TArray<FString>({ "wiz" })) == false);
+	CHECK(FOnlineSubsystemFixtureInvoker::CheckAllTagsIsIn(TestTags, { "bob", "alice", "foo", "wiz" }) == false);
+	CHECK(FOnlineSubsystemFixtureInvoker::CheckAllTagsIsIn(TestTags, { "bob", "alice", "wiz" }) == false);
 
 	// Bounds
-	CHECK(FOnlineSubsystemAutoReg::CheckAllTagsIsIn(TestTags, TArray<FString>({})) == false);
-	CHECK(FOnlineSubsystemAutoReg::CheckAllTagsIsIn({}, TArray<FString>({})) == false);
-	CHECK(FOnlineSubsystemAutoReg::CheckAllTagsIsIn({}, TArray<FString>({ "wiz" })) == false);
-	CHECK(FOnlineSubsystemAutoReg::CheckAllTagsIsIn({ "wiz" }, TArray<FString>({})) == false);
+	CHECK(FOnlineSubsystemFixtureInvoker::CheckAllTagsIsIn(TestTags, TArray<FString>({})) == false);
+	CHECK(FOnlineSubsystemFixtureInvoker::CheckAllTagsIsIn({}, TArray<FString>({})) == false);
+	CHECK(FOnlineSubsystemFixtureInvoker::CheckAllTagsIsIn({}, TArray<FString>({ "wiz" })) == false);
+	CHECK(FOnlineSubsystemFixtureInvoker::CheckAllTagsIsIn({ "wiz" }, TArray<FString>({})) == false);
 }
\ No newline at end of file
