OnlineTests: Renaming the UserNumToLogin, TemplateRequestedNum, LocalUserNum into TestAccountIndex to emphasize the behavior.

diff --git a/Engine/Source/Programs/Online/OSSTestsCore/OnlineSubsystemCatchHelper.cpp b/Engine/Source/Programs/Online/OSSTestsCore/OnlineSubsystemCatchHelper.cpp
--- a/Engine/Source/Programs/Online/OSSTestsCore/OnlineSubsystemCatchHelper.cpp
+++ b/Engine/Source/Programs/Online/OSSTestsCore/OnlineSubsystemCatchHelper.cpp
@@ -108,13 +108,13 @@ FString FOnlineSubsystemTestBaseFixture::GetSubsystem() const
 	return Subsystem;
 }
 
-TArray<FOnlineAccountCredentials> FOnlineSubsystemTestBaseFixture::GetIniCredentials(int32 NumUsers) const
+TArray<FOnlineAccountCredentials> FOnlineSubsystemTestBaseFixture::GetIniCredentials(int32 TestAccountIndex) const
 {
 	FString LoginCredentialCategory = FString::Printf(TEXT("LoginCredentials %s"), *Subsystem);
 	TArray<FString> CredentialsArr;
 	GConfig->GetArray(*LoginCredentialCategory, TEXT("Credentials"), CredentialsArr, GEngineIni);
 
-	if (NumUsers > CredentialsArr.Num())
+	if (TestAccountIndex > CredentialsArr.Num())
 	{
 		UE_LOG(LogOSSTests, Error, TEXT("Attempted to GetCredentials for more than we have stored! Add more credentials to the DefaultEngine.ini for OssTests"));
 		return TArray<FOnlineAccountCredentials>();
@@ -134,12 +134,12 @@ TArray<FOnlineAccountCredentials> FOnlineSubsystemTestBaseFixture::GetIniCredent
 	return OnlineAccountCredentials;
 }
 
-TArray<FOnlineAccountCredentials> FOnlineSubsystemTestBaseFixture::GetCredentials(int32 LocalUserNum, int32 NumUsers) const
+TArray<FOnlineAccountCredentials> FOnlineSubsystemTestBaseFixture::GetCredentials(int32 TestAccountIndex, int32 NumUsers) const
 {
 #if OSSTESTS_USEEXTERNAUTH
-	return CustomCredentials(LocalUserNum, NumUsers);
+	return CustomCredentials(TestAccountIndex, NumUsers);
 #else // OSSTESTS_USEEXTERNAUTH
-	return GetIniCredentials(LocalUserNum);
+	return GetIniCredentials(TestAccountIndex);
 #endif // OSSTESTS_USEEXTERNAUTH
 }
 
diff --git a/Engine/Source/Programs/Online/OSSTestsCore/OnlineSubsystemCatchHelper.h b/Engine/Source/Programs/Online/OSSTestsCore/OnlineSubsystemCatchHelper.h
--- a/Engine/Source/Programs/Online/OSSTestsCore/OnlineSubsystemCatchHelper.h
+++ b/Engine/Source/Programs/Online/OSSTestsCore/OnlineSubsystemCatchHelper.h
@@ -130,12 +130,12 @@ protected:
 	virtual void DestroyCurrentOnlineSubsystemModule() const;
 
 #if OSSTESTS_USEEXTERNAUTH
-	TArray<FOnlineAccountCredentials> CustomCredentials(int32 LocalUserNum, int32 NumUsers) const;
+	TArray<FOnlineAccountCredentials> CustomCredentials(int32 TestAccountIndex, int32 NumUsers) const;
 #endif // OSSTESTS_USEEXTERNAUTH
 
-	TArray<FOnlineAccountCredentials> GetIniCredentials(int32 NumUsers) const;
+	TArray<FOnlineAccountCredentials> GetIniCredentials(int32 TestAccountIndex) const;
 
-	TArray<FOnlineAccountCredentials> GetCredentials(int32 LocalUserNum, int32 NumUsers) const;
+	TArray<FOnlineAccountCredentials> GetCredentials(int32 TestAccountIndex, int32 NumUsers) const;
 
 	FTestPipeline& GetLoginPipeline(uint32 NumUsersToLogin = 1) const;
 	FTestPipeline& GetPipeline() const;
diff --git a/Engine/Source/Programs/Online/OnlineTestsCore/Source/Private/OnlineCatchHelper.cpp b/Engine/Source/Programs/Online/OnlineTestsCore/Source/Private/OnlineCatchHelper.cpp
--- a/Engine/Source/Programs/Online/OnlineTestsCore/Source/Private/OnlineCatchHelper.cpp
+++ b/Engine/Source/Programs/Online/OnlineTestsCore/Source/Private/OnlineCatchHelper.cpp
@@ -78,10 +78,10 @@ SubsystemType OnlineTestBase::GetSubsystem() const
 	return UE::Online::FOnlineServicesRegistry::Get().GetNamedServicesInstance(ServiceType, NAME_None, NAME_None);
 }
 
-bool OnlineTestBase::DeleteAccounts(int32 TemplateRequestedNum) const
+bool OnlineTestBase::DeleteAccounts(int32 TestAccountIndex) const
 {
 #if ONLINETESTS_USEEXTERNAUTH
-	return CustomDeleteAccounts(TemplateRequestedNum);
+	return CustomDeleteAccounts(TestAccountIndex);
 #else // ONLINETESTS_USEEXTERNAUTH
 	return false;
 #endif // ONLINETESTS_USEEXTERNAUTH
@@ -92,10 +92,10 @@ void OnlineTestBase::DestroyCurrentServiceModule() const
 	UE::Online::FOnlineServicesRegistry::Get().DestroyNamedServicesInstance(ServiceType, NAME_None, NAME_None);
 }
 
-bool OnlineTestBase::ResetAccountStatus(int32 TemplateRequestedNum) const
+bool OnlineTestBase::ResetAccountStatus(int32 TestAccountIndex) const
 {
 #if ONLINETESTS_USEEXTERNAUTH
-	return CustomResetAccounts(TemplateRequestedNum);
+	return CustomResetAccounts(TestAccountIndex);
 #else // ONLINETESTS_USEEXTERNAUTH
 	return false;
 #endif // ONLINETESTS_USEEXTERNAUTH
@@ -137,23 +137,23 @@ void OnlineTestBase::UnloadServiceModules()
 	}
 }
 
-TArray<UE::Online::FAuthLogin::Params> OnlineTestBase::GetIniCredentials(int32 LocalUserNum) const
+TArray<UE::Online::FAuthLogin::Params> OnlineTestBase::GetIniCredentials(int32 TestAccountIndex) const
 {
 	FString LoginCredentialCategory = GetLoginCredentialCategory();
 	TArray<FString> LoginCredentialsRaw;
 	GConfig->GetArray(*LoginCredentialCategory, TEXT("Credentials"), LoginCredentialsRaw, GEngineIni);
 
-	if (LocalUserNum >= LoginCredentialsRaw.Num())
+	if (TestAccountIndex >= LoginCredentialsRaw.Num())
 	{
 		UE_LOG_ONLINETESTS(Error, TEXT("Attempted to GetCredentials for more than we have stored! Add more credentials to the DefaultEngine.ini for OnlineTests"));
-		REQUIRE(LocalUserNum >= LoginCredentialsRaw.Num());
+		REQUIRE(TestAccountIndex >= LoginCredentialsRaw.Num());
 		return TArray<UE::Online::FAuthLogin::Params>();
 	}
 
 	TArray<UE::Online::FAuthLogin::Params> AuthLoginParams;
 
 	TArray<FString> LoginCredentialSplit;
-	LoginCredentialsRaw[LocalUserNum].ParseIntoArray(LoginCredentialSplit, TEXT(","), true);
+	LoginCredentialsRaw[TestAccountIndex].ParseIntoArray(LoginCredentialSplit, TEXT(","), true);
 
 	INFO(*FString::Printf(TEXT("Logging in with type %s, id %s, password %s"), *LoginCredentialSplit[0], LoginCredentialSplit.Num() > 1 ? *LoginCredentialSplit[1] : TEXT("UNSET"), LoginCredentialSplit.Num() > 2 ? *LoginCredentialSplit[2] : TEXT("UNSET")));
 	FAuthLogin::Params Params;
@@ -173,12 +173,12 @@ TArray<UE::Online::FAuthLogin::Params> OnlineTestBase::GetIniCredentials(int32 L
 	return AuthLoginParams;
 }
 
-TArray<UE::Online::FAuthLogin::Params> OnlineTestBase::GetCredentials(int32 LocalUserNum, int32 NumUsers) const
+TArray<UE::Online::FAuthLogin::Params> OnlineTestBase::GetCredentials(int32 TestAccountIndex, int32 NumUsers) const
 {
 #if ONLINETESTS_USEEXTERNAUTH
-	return CustomCredentials(LocalUserNum, NumUsers);
+	return CustomCredentials(TestAccountIndex, NumUsers);
 #else // ONLINETESTS_USEEXTERNAUTH
-	return GetIniCredentials(LocalUserNum);
+	return GetIniCredentials(TestAccountIndex);
 #endif // ONLINETESTS_USEEXTERNAUTH
 }
 
@@ -201,7 +201,7 @@ FTestPipeline& OnlineTestBase::GetLoginPipeline(std::initializer_list<std::refer
 	return GetLoginPipeline(0, AccountIds);
 }
 
-FTestPipeline& OnlineTestBase::GetLoginPipeline(uint32 UserNumToLogin, std::initializer_list<std::reference_wrapper<FAccountId>> AccountIds) const
+FTestPipeline& OnlineTestBase::GetLoginPipeline(uint32 TestAccountIndex, std::initializer_list<std::reference_wrapper<FAccountId>> AccountIds) const
 {
 	// Don't call GetLoginPipeline more than once per test
 	REQUIRE(NumLocalUsers == -1);
@@ -234,9 +234,9 @@ FTestPipeline& OnlineTestBase::GetLoginPipeline(uint32 UserNumToLogin, std::init
 	}
 	else if (NumLocalUsers > 0)
 	{
-		UserNumToLogout = UserNumToLogin;
+		TestAccountID = TestAccountIndex;
 
-		TArray<UE::Online::FAuthLogin::Params> AuthLoginParams = GetCredentials(UserNumToLogin, NumLocalUsers);
+		TArray<UE::Online::FAuthLogin::Params> AuthLoginParams = GetCredentials(TestAccountIndex, NumLocalUsers);
 		for (UE::Online::FAuthLogin::Params AuthLoginParam : AuthLoginParams)
 		{
 			Pipeline->EmplaceStep<FAuthLoginStep>(MoveTemp(AuthLoginParam));
@@ -285,7 +285,7 @@ void OnlineTestBase::RunToCompletion(bool bLogout) const
 		}
 		else if(NumLocalUsers > 0)
 		{
-			TArray<UE::Online::FAuthLogin::Params> AuthLoginParams = GetCredentials(UserNumToLogout, NumLocalUsers);
+			TArray<UE::Online::FAuthLogin::Params> AuthLoginParams = GetCredentials(TestAccountID, NumLocalUsers);
 			for (UE::Online::FAuthLogin::Params AuthLoginParam : AuthLoginParams)
 			{
 				Pipeline->EmplaceStep<FAuthLogoutStep>(AuthLoginParam.PlatformUserId);
diff --git a/Engine/Source/Programs/Online/OnlineTestsCore/Source/Private/Tests/Lobby/LobbyCreateInviteTests.cpp b/Engine/Source/Programs/Online/OnlineTestsCore/Source/Private/Tests/Lobby/LobbyCreateInviteTests.cpp
--- a/Engine/Source/Programs/Online/OnlineTestsCore/Source/Private/Tests/Lobby/LobbyCreateInviteTests.cpp
+++ b/Engine/Source/Programs/Online/OnlineTestsCore/Source/Private/Tests/Lobby/LobbyCreateInviteTests.cpp
@@ -8,7 +8,7 @@
 
 LOBBY_CREATE_INVITE_TEST_CASE("Basic lobby create and join test")
 {
-	int32 UserNumToLogin = 1;
+	int32 TestAccountIndex = 1;
 	FAccountId AccountId;
 
 	UE::Online::FCreateLobby::Params LobbyCreateParams;
@@ -21,7 +21,7 @@ LOBBY_CREATE_INVITE_TEST_CASE("Basic lobby create and join test")
 	FJoinLobby::Params JoinLobbyParams;
 	JoinLobbyParams.LocalName = TEXT("TestLobby");
 
-	FTestPipeline& LoginPipeline = GetLoginPipeline(UserNumToLogin, { AccountId });
+	FTestPipeline& LoginPipeline = GetLoginPipeline(TestAccountIndex, { AccountId });
 
 	LobbyCreateParams.LocalAccountId = AccountId;
 	JoinLobbyParams.LocalAccountId = AccountId;
diff --git a/Engine/Source/Programs/Online/OnlineTestsCore/Source/Private/Tests/Sessions/AddRemoveSessionMemberTests.cpp b/Engine/Source/Programs/Online/OnlineTestsCore/Source/Private/Tests/Sessions/AddRemoveSessionMemberTests.cpp
--- a/Engine/Source/Programs/Online/OnlineTestsCore/Source/Private/Tests/Sessions/AddRemoveSessionMemberTests.cpp
+++ b/Engine/Source/Programs/Online/OnlineTestsCore/Source/Private/Tests/Sessions/AddRemoveSessionMemberTests.cpp
@@ -157,7 +157,7 @@ SESSIONS_TEST_CASE("If I call RemoveSessionMember with valid data, the operation
 {
 	DestroyCurrentServiceModule();
 
-	int32 UserNumToLogin = 7;
+	int32 TestAccountIndex = 7;
 	FAccountId AccountId;
 
 	FCreateSession::Params OpCreateParams;
@@ -184,7 +184,7 @@ SESSIONS_TEST_CASE("If I call RemoveSessionMember with valid data, the operation
 	RemoveSessionMemberHelperParams.OpParams = &OpRemoveParams;
 	RemoveSessionMemberHelperParams.OpParams->SessionName = TEXT("SessionNameValidRemoveMember");
 
-	FTestPipeline& LoginPipeline = GetLoginPipeline(UserNumToLogin, { AccountId });
+	FTestPipeline& LoginPipeline = GetLoginPipeline(TestAccountIndex, { AccountId });
 
 	CreateSessionHelperParams.OpParams->LocalAccountId = AccountId;
 	AddSessionMemberHelperParams.OpParams->LocalAccountId = AccountId;
diff --git a/Engine/Source/Programs/Online/OnlineTestsCore/Source/Private/Tests/Sessions/CreateSessionTests.cpp b/Engine/Source/Programs/Online/OnlineTestsCore/Source/Private/Tests/Sessions/CreateSessionTests.cpp
--- a/Engine/Source/Programs/Online/OnlineTestsCore/Source/Private/Tests/Sessions/CreateSessionTests.cpp
+++ b/Engine/Source/Programs/Online/OnlineTestsCore/Source/Private/Tests/Sessions/CreateSessionTests.cpp
@@ -111,7 +111,7 @@ SESSIONS_TEST_CASE("If I call CreateSession with session override id less than 1
 {
 	LOG_SCOPE_VERBOSITY_OVERRIDE(LogOnlineServices, ELogVerbosity::NoLogging);
 
-	int32 UserNumToLogin = 7;
+	int32 TestAccountIndex = 7;
 	FAccountId AccountId;
 
 	FCreateSession::Params OpFirstCreateParams;
@@ -134,7 +134,7 @@ SESSIONS_TEST_CASE("If I call CreateSession with session override id less than 1
 	SecondCreateSessionHelperParams.OpParams->SessionSettings.NumMaxConnections = 2;
 	SecondCreateSessionHelperParams.ExpectedError = TOnlineResult<FCreateSession>(Errors::InvalidParams());
 
-	FTestPipeline& LoginPipeline = GetLoginPipeline(UserNumToLogin, { AccountId });
+	FTestPipeline& LoginPipeline = GetLoginPipeline(TestAccountIndex, { AccountId });
 
 	FirstCreateSessionHelperParams.OpParams->LocalAccountId = AccountId;
 	SecondCreateSessionHelperParams.OpParams->LocalAccountId = AccountId;
@@ -150,7 +150,7 @@ SESSIONS_TEST_CASE("If I call CreateSession with an name already in use, I get a
 {
 	DestroyCurrentServiceModule();
 	
-	int32 UserNumToLogin = 7;
+	int32 TestAccountIndex = 7;
 	FAccountId AccountId;
 
 	FCreateSession::Params OpCreateParams;
@@ -172,7 +172,7 @@ SESSIONS_TEST_CASE("If I call CreateSession with an name already in use, I get a
 	LeaveSessionHelperParams.OpParams->SessionName = TEXT("SessionNameInUse");
 	LeaveSessionHelperParams.OpParams->bDestroySession = true;
 
-	FTestPipeline& LoginPipeline = GetLoginPipeline(UserNumToLogin, { AccountId });
+	FTestPipeline& LoginPipeline = GetLoginPipeline(TestAccountIndex, { AccountId });
 
 	FirstCreateSessionHelperParams.OpParams->LocalAccountId = AccountId;
 	SecondCreateSessionHelperParams.OpParams->LocalAccountId = AccountId;
@@ -193,7 +193,7 @@ SESSIONS_TEST_CASE("If I call CreateSession with the presence flag set after the
 
 	DestroyCurrentServiceModule();
 	
-	int32 UserNumToLogin = 7;
+	int32 TestAccountIndex = 7;
 	FAccountId AccountId;
 
 	FCreateSession::Params OpFirstCreateParams;
@@ -219,7 +219,7 @@ SESSIONS_TEST_CASE("If I call CreateSession with the presence flag set after the
 	LeaveSessionHelperParams.OpParams->SessionName = TEXT("SessionNamePresenceEnabled1");
 	LeaveSessionHelperParams.OpParams->bDestroySession = true;
 
-	FTestPipeline& LoginPipeline = GetLoginPipeline(UserNumToLogin, { AccountId });
+	FTestPipeline& LoginPipeline = GetLoginPipeline(TestAccountIndex, { AccountId });
 
 	FirstCreateSessionHelperParams.OpParams->LocalAccountId = AccountId;
 	SecondCreateSessionHelperParams.OpParams->LocalAccountId = AccountId;
diff --git a/Engine/Source/Programs/Online/OnlineTestsCore/Source/Private/Tests/Sessions/FindSessionsTests.cpp b/Engine/Source/Programs/Online/OnlineTestsCore/Source/Private/Tests/Sessions/FindSessionsTests.cpp
--- a/Engine/Source/Programs/Online/OnlineTestsCore/Source/Private/Tests/Sessions/FindSessionsTests.cpp
+++ b/Engine/Source/Programs/Online/OnlineTestsCore/Source/Private/Tests/Sessions/FindSessionsTests.cpp
@@ -138,7 +138,7 @@ SESSIONS_TEST_CASE("If I call FindSessions while a search is already running for
 	DestroyCurrentServiceModule();
 	ReturnAccounts();
 
-	int32 UserNumToLogin = 7;
+	int32 TestAccountIndex = 7;
 	FAccountId FirstAccountId, SecondAccountId;
 
 	FCreateSession::Params OpCreateParams;
@@ -171,7 +171,7 @@ SESSIONS_TEST_CASE("If I call FindSessions while a search is already running for
 	LeaveSessionHelperParams.OpParams->SessionName = TEXT("AlreadyPendingName");
 	LeaveSessionHelperParams.OpParams->bDestroySession = true;
 
-	FTestPipeline& LoginPipeline = GetLoginPipeline(UserNumToLogin, { FirstAccountId, SecondAccountId });
+	FTestPipeline& LoginPipeline = GetLoginPipeline(TestAccountIndex, { FirstAccountId, SecondAccountId });
 
 	CreateSessionHelperParams.OpParams->LocalAccountId = FirstAccountId;
 
diff --git a/Engine/Source/Programs/Online/OnlineTestsCore/Source/Private/Tests/Sessions/GetSessionsTests.cpp b/Engine/Source/Programs/Online/OnlineTestsCore/Source/Private/Tests/Sessions/GetSessionsTests.cpp
--- a/Engine/Source/Programs/Online/OnlineTestsCore/Source/Private/Tests/Sessions/GetSessionsTests.cpp
+++ b/Engine/Source/Programs/Online/OnlineTestsCore/Source/Private/Tests/Sessions/GetSessionsTests.cpp
@@ -425,7 +425,7 @@ SESSIONS_TEST_CASE("If I call GetSessionInviteById with a valid account id, but
 {
 	DestroyCurrentServiceModule();
 
-	int32 UserNumToLogin = 7;
+	int32 TestAccountIndex = 7;
 	FAccountId FirstAccountId, SecondAccountId;
 
 	FCreateSession::Params OpCreateParams;
@@ -449,7 +449,7 @@ SESSIONS_TEST_CASE("If I call GetSessionInviteById with a valid account id, but
 	LeaveSessionHelperParams.OpParams->SessionName = TEXT("GetInviteByIdWithValidAccountIdName");
 	LeaveSessionHelperParams.OpParams->bDestroySession = true;
 
-	FTestPipeline& LoginPipeline = GetLoginPipeline(UserNumToLogin, { FirstAccountId, SecondAccountId });
+	FTestPipeline& LoginPipeline = GetLoginPipeline(TestAccountIndex, { FirstAccountId, SecondAccountId });
 
 	OpGetInviteByIdParams.LocalAccountId = FirstAccountId;
 	CreateSessionHelperParams.OpParams->LocalAccountId = FirstAccountId;
@@ -484,7 +484,7 @@ SESSIONS_TEST_CASE("If I call GetSessionInviteById with valid invite id, but wit
 
 	DestroyCurrentServiceModule();
 
-	int32 UserNumToLogin = 7;
+	int32 TestAccountIndex = 7;
 	FAccountId FirstAccountId, SecondAccountId;
 
 	FCreateSession::Params OpCreateParams;
@@ -512,7 +512,7 @@ SESSIONS_TEST_CASE("If I call GetSessionInviteById with valid invite id, but wit
 	LeaveSessionHelperParams.OpParams->SessionName = TEXT("GetInviteByIdWithValidInviteIdName");
 	LeaveSessionHelperParams.OpParams->bDestroySession = true;
 
-	FTestPipeline& LoginPipeline = GetLoginPipeline(UserNumToLogin, { FirstAccountId, SecondAccountId });
+	FTestPipeline& LoginPipeline = GetLoginPipeline(TestAccountIndex, { FirstAccountId, SecondAccountId });
 
 	OpGetInviteByIdParams.LocalAccountId = SecondAccountId;
 	CreateSessionHelperParams.OpParams->LocalAccountId = FirstAccountId;
@@ -548,7 +548,7 @@ SESSIONS_TEST_CASE("If I call GetSessionInviteById with a valid data, I get a va
 {
 	DestroyCurrentServiceModule();
 
-	int32 UserNumToLogin = 7;
+	int32 TestAccountIndex = 7;
 	FAccountId FirstAccountId, SecondAccountId;
 
 	FCreateSession::Params OpCreateParams;
@@ -572,7 +572,7 @@ SESSIONS_TEST_CASE("If I call GetSessionInviteById with a valid data, I get a va
 	LeaveSessionHelperParams.OpParams->SessionName = TEXT("GetInviteByIdValidName");
 	LeaveSessionHelperParams.OpParams->bDestroySession = true;
 
-	FTestPipeline& LoginPipeline = GetLoginPipeline(UserNumToLogin, { FirstAccountId, SecondAccountId });
+	FTestPipeline& LoginPipeline = GetLoginPipeline(TestAccountIndex, { FirstAccountId, SecondAccountId });
 
 	OpGetInviteByIdParams.LocalAccountId = SecondAccountId;
 	CreateSessionHelperParams.OpParams->LocalAccountId = FirstAccountId;
diff --git a/Engine/Source/Programs/Online/OnlineTestsCore/Source/Private/Tests/Sessions/JoinSessionTests.cpp b/Engine/Source/Programs/Online/OnlineTestsCore/Source/Private/Tests/Sessions/JoinSessionTests.cpp
--- a/Engine/Source/Programs/Online/OnlineTestsCore/Source/Private/Tests/Sessions/JoinSessionTests.cpp
+++ b/Engine/Source/Programs/Online/OnlineTestsCore/Source/Private/Tests/Sessions/JoinSessionTests.cpp
@@ -57,7 +57,7 @@ SESSIONS_TEST_CASE("If I call JoinSession with an empty session name, I get an e
 
 	DestroyCurrentServiceModule();
 
-	int32 UserNumToLogin = 7;
+	int32 TestAccountIndex = 7;
 	FAccountId AccountId;
 
 	FCreateSession::Params OpCreateParams;
@@ -80,7 +80,7 @@ SESSIONS_TEST_CASE("If I call JoinSession with an empty session name, I get an e
 	JoinSessionHelperParams.OpParams = &OpJoinParams;
 	JoinSessionHelperParams.ExpectedError = TOnlineResult<FJoinSession>(Errors::InvalidParams());
 
-	FTestPipeline& LoginPipeline = GetLoginPipeline(UserNumToLogin, { AccountId });
+	FTestPipeline& LoginPipeline = GetLoginPipeline(TestAccountIndex, { AccountId });
 
 	CreateSessionHelperParams.OpParams->LocalAccountId = AccountId;
 	JoinSessionHelperParams.OpParams->LocalAccountId = AccountId;
@@ -108,7 +108,7 @@ SESSIONS_TEST_CASE("If I call JoinSession with an name already in use, I get an
 
 	DestroyCurrentServiceModule();
 
-	int32 UserNumToLogin = 7;
+	int32 TestAccountIndex = 7;
 	FAccountId FirstAccountId, SecondAccountId;
 
 	FCreateSession::Params OpFirstCreateParams;
@@ -165,7 +165,7 @@ SESSIONS_TEST_CASE("If I call JoinSession with an name already in use, I get an
 	SecondLeaveSessionHelperParams.OpParams->SessionName = TEXT("SessionNameInUseJoin1");
 	SecondLeaveSessionHelperParams.OpParams->bDestroySession = true;
 
-	FTestPipeline& LoginPipeline = GetLoginPipeline(UserNumToLogin, { FirstAccountId, SecondAccountId });
+	FTestPipeline& LoginPipeline = GetLoginPipeline(TestAccountIndex, { FirstAccountId, SecondAccountId });
 
 	FirstCreateSessionHelperParams.OpParams->LocalAccountId = FirstAccountId;
 	SecondCreateSessionHelperParams.OpParams->LocalAccountId = FirstAccountId;
@@ -226,7 +226,7 @@ SESSIONS_TEST_CASE("If I call JoinSession with a valid but unregistered session
 
 	DestroyCurrentServiceModule();
 
-	int32 UserNumToLogin = 7;
+	int32 TestAccountIndex = 7;
 	FAccountId FirstAccountId, SecondAccountId;
 
 	FCreateSession::Params OpCreateParams;
@@ -272,7 +272,7 @@ SESSIONS_TEST_CASE("If I call JoinSession with a valid but unregistered session
 	SecondLeaveSessionHelperParams.OpParams->SessionName = TEXT("SessionIdNotFoundJoinName");
 	SecondLeaveSessionHelperParams.ExpectedError = TOnlineResult<FLeaveSession>(Errors::InvalidState());
 
-	FTestPipeline& LoginPipeline = GetLoginPipeline(UserNumToLogin, { FirstAccountId, SecondAccountId });
+	FTestPipeline& LoginPipeline = GetLoginPipeline(TestAccountIndex, { FirstAccountId, SecondAccountId });
 
 	CreateSessionHelperParams.OpParams->LocalAccountId = FirstAccountId;
 
@@ -394,7 +394,7 @@ SESSIONS_TEST_CASE("If I call JoinSession for a session that is not joinable, I
 {
 	DestroyCurrentServiceModule();
 
-	int32 UserNumToLogin = 7;
+	int32 TestAccountIndex = 7;
 	FAccountId FirstAccountId, SecondAccountId;
 
 	FCreateSession::Params OpCreateParams;
@@ -430,7 +430,7 @@ SESSIONS_TEST_CASE("If I call JoinSession for a session that is not joinable, I
 	JoinSessionHelperParams.OpParams->SessionName = TEXT("SessionNotJoinableName1");
 	JoinSessionHelperParams.ExpectedError = TOnlineResult<FJoinSession>(Errors::AccessDenied());
 
-	FTestPipeline& LoginPipeline = GetLoginPipeline(UserNumToLogin, { FirstAccountId, SecondAccountId });
+	FTestPipeline& LoginPipeline = GetLoginPipeline(TestAccountIndex, { FirstAccountId, SecondAccountId });
 
 	CreateSessionHelperParams.OpParams->LocalAccountId = FirstAccountId;
 	UpdateSessionSettingsHelperParams.OpParams->LocalAccountId = FirstAccountId;
diff --git a/Engine/Source/Programs/Online/OnlineTestsCore/Source/Private/Tests/Sessions/LeaveSessionTests.cpp b/Engine/Source/Programs/Online/OnlineTestsCore/Source/Private/Tests/Sessions/LeaveSessionTests.cpp
--- a/Engine/Source/Programs/Online/OnlineTestsCore/Source/Private/Tests/Sessions/LeaveSessionTests.cpp
+++ b/Engine/Source/Programs/Online/OnlineTestsCore/Source/Private/Tests/Sessions/LeaveSessionTests.cpp
@@ -65,7 +65,7 @@ SESSIONS_TEST_CASE("If I call LeaveSession with valid data, the operation comple
 {
 	DestroyCurrentServiceModule();
 
-	int32 UserNumToLogin = 7;
+	int32 TestAccountIndex = 7;
 	FAccountId AccountId;
 
 	FCreateSession::Params OpCreateParams;
@@ -82,7 +82,7 @@ SESSIONS_TEST_CASE("If I call LeaveSession with valid data, the operation comple
 	LeaveSessionHelperParams.OpParams->SessionName = TEXT("LeaveSessionNameValidLeave");
 	LeaveSessionHelperParams.OpParams->bDestroySession = true;
 
-	FTestPipeline& LoginPipeline = GetLoginPipeline(UserNumToLogin, { AccountId });
+	FTestPipeline& LoginPipeline = GetLoginPipeline(TestAccountIndex, { AccountId });
 
 	CreateSessionHelperParams.OpParams->LocalAccountId = AccountId;
 	LeaveSessionHelperParams.OpParams->LocalAccountId = AccountId;
diff --git a/Engine/Source/Programs/Online/OnlineTestsCore/Source/Private/Tests/Sessions/PresenceSessionTests.cpp b/Engine/Source/Programs/Online/OnlineTestsCore/Source/Private/Tests/Sessions/PresenceSessionTests.cpp
--- a/Engine/Source/Programs/Online/OnlineTestsCore/Source/Private/Tests/Sessions/PresenceSessionTests.cpp
+++ b/Engine/Source/Programs/Online/OnlineTestsCore/Source/Private/Tests/Sessions/PresenceSessionTests.cpp
@@ -58,7 +58,7 @@ SESSIONS_TEST_CASE("If I call IsPresenceSession with an unregistered account id,
 
 	DestroyCurrentServiceModule();
 
-	int32 UserNumToLogin = 7;
+	int32 TestAccountIndex = 7;
 	FAccountId AccountId;
 
 	FCreateSession::Params OpCreateParams;
@@ -80,7 +80,7 @@ SESSIONS_TEST_CASE("If I call IsPresenceSession with an unregistered account id,
 	FGetSessionByName::Params OpGetByNameParams;
 	OpGetByNameParams.LocalName = TEXT("IsPresenceUnregisteredSessionName");
 
-	FTestPipeline& LoginPipeline = GetLoginPipeline(UserNumToLogin, { AccountId });
+	FTestPipeline& LoginPipeline = GetLoginPipeline(TestAccountIndex, { AccountId });
 
 	CreateSessionHelperParams.OpParams->LocalAccountId = AccountId;
 	LeaveSessionHelperParams.OpParams->LocalAccountId = AccountId;
diff --git a/Engine/Source/Programs/Online/OnlineTestsCore/Source/Private/Tests/Sessions/SendRejectSessionInvitesTests.cpp b/Engine/Source/Programs/Online/OnlineTestsCore/Source/Private/Tests/Sessions/SendRejectSessionInvitesTests.cpp
--- a/Engine/Source/Programs/Online/OnlineTestsCore/Source/Private/Tests/Sessions/SendRejectSessionInvitesTests.cpp
+++ b/Engine/Source/Programs/Online/OnlineTestsCore/Source/Private/Tests/Sessions/SendRejectSessionInvitesTests.cpp
@@ -75,7 +75,7 @@ SESSIONS_TEST_CASE("If I call SendSessionInvite with an empty target users, I ge
 {
 	DestroyCurrentServiceModule();
 
-	int32 UserNumToLogin = 7;
+	int32 TestAccountIndex = 7;
 	FAccountId FirstAccountId, SecondAccountId;
 
 	FCreateSession::Params OpCreateParams;
@@ -98,7 +98,7 @@ SESSIONS_TEST_CASE("If I call SendSessionInvite with an empty target users, I ge
 	LeaveSessionHelperParams.OpParams->SessionName = TEXT("SessionSendInviteEmptyName");
 	LeaveSessionHelperParams.OpParams->bDestroySession = true;
 
-	FTestPipeline& LoginPipeline = GetLoginPipeline(UserNumToLogin, { FirstAccountId, SecondAccountId });
+	FTestPipeline& LoginPipeline = GetLoginPipeline(TestAccountIndex, { FirstAccountId, SecondAccountId });
 
 	CreateSessionHelperParams.OpParams->LocalAccountId = FirstAccountId;
 	SendSessionInviteHelperParams.OpParams->LocalAccountId = FirstAccountId;
@@ -116,7 +116,7 @@ SESSIONS_TEST_CASE("If I call SendSessionInvite with an invalid target users, I
 {
 	DestroyCurrentServiceModule();
 
-	int32 UserNumToLogin = 7;
+	int32 TestAccountIndex = 7;
 	FAccountId FirstAccountId, SecondAccountId;
 		
 	FCreateSession::Params OpCreateParams;
@@ -140,7 +140,7 @@ SESSIONS_TEST_CASE("If I call SendSessionInvite with an invalid target users, I
 	LeaveSessionHelperParams.OpParams->SessionName = TEXT("SessionSendInviteInvalidUsersName");
 	LeaveSessionHelperParams.OpParams->bDestroySession = true;
 
-	FTestPipeline& LoginPipeline = GetLoginPipeline(UserNumToLogin, { FirstAccountId, SecondAccountId });
+	FTestPipeline& LoginPipeline = GetLoginPipeline(TestAccountIndex, { FirstAccountId, SecondAccountId });
 
 	CreateSessionHelperParams.OpParams->LocalAccountId = FirstAccountId;
 	SendSessionInviteHelperParams.OpParams->LocalAccountId = FirstAccountId;
@@ -158,7 +158,7 @@ SESSIONS_TEST_CASE("If I call SendSessionInvite with valid data, the operation c
 {
 	DestroyCurrentServiceModule();
 
-	int32 UserNumToLogin = 7;
+	int32 TestAccountIndex = 7;
 	FAccountId FirstAccountId, SecondAccountId;
 
 	FCreateSession::Params OpCreateParams;
@@ -182,7 +182,7 @@ SESSIONS_TEST_CASE("If I call SendSessionInvite with valid data, the operation c
 	LeaveSessionHelperParams.OpParams->SessionName = TEXT("SessionSendInviteValidName");
 	LeaveSessionHelperParams.OpParams->bDestroySession = true;
 
-	FTestPipeline& LoginPipeline = GetLoginPipeline(UserNumToLogin, { FirstAccountId, SecondAccountId });
+	FTestPipeline& LoginPipeline = GetLoginPipeline(TestAccountIndex, { FirstAccountId, SecondAccountId });
 
 	OpGetAllSessionInvitesParams.LocalAccountId = SecondAccountId;
 	CreateSessionHelperParams.OpParams->LocalAccountId = FirstAccountId;
@@ -254,7 +254,7 @@ SESSIONS_TEST_CASE("If I call RejectSessionInvite with valid data, the operation
 
 	DestroyCurrentServiceModule();
 
-	int32 UserNumToLogin = 7;
+	int32 TestAccountIndex = 7;
 	FAccountId FirstAccountId, SecondAccountId;
 
 	FCreateSession::Params OpCreateParams;
@@ -282,7 +282,7 @@ SESSIONS_TEST_CASE("If I call RejectSessionInvite with valid data, the operation
 	LeaveSessionHelperParams.OpParams->SessionName = TEXT("SessionRejectInviteValidName");
 	LeaveSessionHelperParams.OpParams->bDestroySession = true;
 
-	FTestPipeline& LoginPipeline = GetLoginPipeline(UserNumToLogin, { FirstAccountId, SecondAccountId });
+	FTestPipeline& LoginPipeline = GetLoginPipeline(TestAccountIndex, { FirstAccountId, SecondAccountId });
 
 	OpGetAllSessionInvitesParams.LocalAccountId = SecondAccountId;
 	CreateSessionHelperParams.OpParams->LocalAccountId = FirstAccountId;
diff --git a/Engine/Source/Programs/Online/OnlineTestsCore/Source/Private/Tests/Sessions/StartMatchmakingTests.cpp b/Engine/Source/Programs/Online/OnlineTestsCore/Source/Private/Tests/Sessions/StartMatchmakingTests.cpp
--- a/Engine/Source/Programs/Online/OnlineTestsCore/Source/Private/Tests/Sessions/StartMatchmakingTests.cpp
+++ b/Engine/Source/Programs/Online/OnlineTestsCore/Source/Private/Tests/Sessions/StartMatchmakingTests.cpp
@@ -177,7 +177,7 @@ SESSIONS_TEST_CASE("If I call StartMatchmaking with name already in use, I get a
 {
 	DestroyCurrentServiceModule();
 
-	int32 UserNumToLogin = 7;
+	int32 TestAccountIndex = 7;
 	FAccountId AccountId;
 
 	FCreateSession::Params OpCreateParams;
@@ -202,7 +202,7 @@ SESSIONS_TEST_CASE("If I call StartMatchmaking with name already in use, I get a
 	LeaveSessionHelperParams.OpParams->SessionName = TEXT("StartMatchmakingAlreadyInUseSessionName");
 	LeaveSessionHelperParams.OpParams->bDestroySession = true;
 
-	FTestPipeline& LoginPipeline = GetLoginPipeline(UserNumToLogin, { AccountId });
+	FTestPipeline& LoginPipeline = GetLoginPipeline(TestAccountIndex, { AccountId });
 
 	CreateSessionHelperParams.OpParams->LocalAccountId = AccountId;
 	StartMatchmakingHelperParams.OpParams->SessionCreationParameters.LocalAccountId = AccountId;
diff --git a/Engine/Source/Programs/Online/OnlineTestsCore/Source/Private/Tests/Sessions/UpdateSessionSettingsTests.cpp b/Engine/Source/Programs/Online/OnlineTestsCore/Source/Private/Tests/Sessions/UpdateSessionSettingsTests.cpp
--- a/Engine/Source/Programs/Online/OnlineTestsCore/Source/Private/Tests/Sessions/UpdateSessionSettingsTests.cpp
+++ b/Engine/Source/Programs/Online/OnlineTestsCore/Source/Private/Tests/Sessions/UpdateSessionSettingsTests.cpp
@@ -130,7 +130,7 @@ UPDATESESSIONS_TEST_CASE("If I call UpdateSessionSettings with valid data, the o
 {
 	DestroyCurrentServiceModule();
 
-	int32 UserNumToLogin = 7;
+	int32 TestAccountIndex = 7;
 	FAccountId AccountId;
 
 	FCreateSession::Params OpCreateParams;
@@ -164,7 +164,7 @@ UPDATESESSIONS_TEST_CASE("If I call UpdateSessionSettings with valid data, the o
 	LeaveSessionHelperParams.OpParams->SessionName = TEXT("ValidUpdateSessionName");
 	LeaveSessionHelperParams.OpParams->bDestroySession = true;
 
-	FTestPipeline& LoginPipeline = GetLoginPipeline(UserNumToLogin, { AccountId });
+	FTestPipeline& LoginPipeline = GetLoginPipeline(TestAccountIndex, { AccountId });
 
 	CreateSessionHelperParams.OpParams->LocalAccountId = AccountId;
 	UpdateSessionSettingsHelperParams.OpParams->LocalAccountId = AccountId;
diff --git a/Engine/Source/Programs/Online/OnlineTestsCore/Source/Private/Tests/Social/AcceptFriendInviteTests.cpp b/Engine/Source/Programs/Online/OnlineTestsCore/Source/Private/Tests/Social/AcceptFriendInviteTests.cpp
--- a/Engine/Source/Programs/Online/OnlineTestsCore/Source/Private/Tests/Social/AcceptFriendInviteTests.cpp
+++ b/Engine/Source/Programs/Online/OnlineTestsCore/Source/Private/Tests/Social/AcceptFriendInviteTests.cpp
@@ -48,7 +48,7 @@ SOCIAL_TEST_CASE("Verify that AcceptFriendInvite returns a fail message if the l
 {
 	FAccountId FirstAccountId, SecondAccountId;
 
-	int32 UserNumToLogin = 1;
+	int32 TestAccountIndex = 1;
 	TSharedPtr<FPlatformUserId> FirstAccountPlatformUserId = MakeShared<FPlatformUserId>();
 	TSharedPtr<FPlatformUserId> SecondAccountPlatformUserId = MakeShared<FPlatformUserId>();
 
@@ -59,7 +59,7 @@ SOCIAL_TEST_CASE("Verify that AcceptFriendInvite returns a fail message if the l
 	AcceptInviteHelperParams.OpParams = &OpAcceptFriendInviteParams;
 	AcceptInviteHelperParams.ExpectedError = TOnlineResult<FAcceptFriendInvite>(Errors::NotLoggedIn());
 
-	FTestPipeline& LoginPipeline = GetLoginPipeline(UserNumToLogin, { FirstAccountId, SecondAccountId });
+	FTestPipeline& LoginPipeline = GetLoginPipeline(TestAccountIndex, { FirstAccountId, SecondAccountId });
 
 	AcceptInviteHelperParams.OpParams->LocalAccountId = FirstAccountId;
 	AcceptInviteHelperParams.OpParams->TargetAccountId = SecondAccountId;
diff --git a/Engine/Source/Programs/Online/OnlineTestsCore/Source/Private/Tests/Social/BlockUserTests.cpp b/Engine/Source/Programs/Online/OnlineTestsCore/Source/Private/Tests/Social/BlockUserTests.cpp
--- a/Engine/Source/Programs/Online/OnlineTestsCore/Source/Private/Tests/Social/BlockUserTests.cpp
+++ b/Engine/Source/Programs/Online/OnlineTestsCore/Source/Private/Tests/Social/BlockUserTests.cpp
@@ -48,7 +48,7 @@ SOCIAL_TEST_CASE("Verify that BlockUser returns a fail message if the local user
 {
 	FAccountId FirstAccountId, SecondAccountId;
 	
-	int32 UserNumToLogin = 1;
+	int32 TestAccountIndex = 1;
 	TSharedPtr<FPlatformUserId> FirstAccountPlatformUserId = MakeShared<FPlatformUserId>();
 	TSharedPtr<FPlatformUserId> SecondAccountPlatformUserId = MakeShared<FPlatformUserId>();
 
@@ -59,7 +59,7 @@ SOCIAL_TEST_CASE("Verify that BlockUser returns a fail message if the local user
 	BlockUserHelperHelperParams.OpParams = &OpBlockUserParams;
 	BlockUserHelperHelperParams.ExpectedError = TOnlineResult<FBlockUser>(Errors::NotLoggedIn());
 
-	FTestPipeline& LoginPipeline = GetLoginPipeline(UserNumToLogin, { FirstAccountId, SecondAccountId });
+	FTestPipeline& LoginPipeline = GetLoginPipeline(TestAccountIndex, { FirstAccountId, SecondAccountId });
 
 	BlockUserHelperHelperParams.OpParams->LocalAccountId = FirstAccountId;
 	BlockUserHelperHelperParams.OpParams->TargetAccountId = SecondAccountId;
diff --git a/Engine/Source/Programs/Online/OnlineTestsCore/Source/Private/Tests/Social/GetFriendsTests.cpp b/Engine/Source/Programs/Online/OnlineTestsCore/Source/Private/Tests/Social/GetFriendsTests.cpp
--- a/Engine/Source/Programs/Online/OnlineTestsCore/Source/Private/Tests/Social/GetFriendsTests.cpp
+++ b/Engine/Source/Programs/Online/OnlineTestsCore/Source/Private/Tests/Social/GetFriendsTests.cpp
@@ -14,14 +14,14 @@ SOCIAL_TEST_CASE("Verify that GetFriends returns an empty list if there are no c
 {
 	FAccountId AccountId;	
 
-	int32 UserNumToLogin = 5;
+	int32 TestAccountIndex = 5;
 
 	FGetFriends::Params OpGetParams;
 	FGetFriendsHelper::FHelperParams GetFriendsHelperParams;
 	GetFriendsHelperParams.OpParams = &OpGetParams;
 	GetFriendsHelperParams.ExpectedError = TOnlineResult<FGetFriends>(Errors::InvalidState());
 
-	FTestPipeline& LoginPipeline = GetLoginPipeline(UserNumToLogin, { AccountId });
+	FTestPipeline& LoginPipeline = GetLoginPipeline(TestAccountIndex, { AccountId });
 
 	GetFriendsHelperParams.OpParams->LocalAccountId = AccountId;
 
@@ -35,7 +35,7 @@ SOCIAL_TEST_CASE("Verify that GetFriends returns a list of 1 Friend if there is
 {
 	FAccountId AccountId;
 	
-	int32 UserNumToLogin = 6;
+	int32 TestAccountIndex = 6;
 
 	FQueryFriends::Params OpQueryParams;
 	FQueryFriendsHelper::FHelperParams QueryFriendsHelperParams;
@@ -45,7 +45,7 @@ SOCIAL_TEST_CASE("Verify that GetFriends returns a list of 1 Friend if there is
 	FGetFriendsHelper::FHelperParams GetFriendsHelperParams;
 	GetFriendsHelperParams.OpParams = &OpGetParams;
 
-	FTestPipeline& LoginPipeline = GetLoginPipeline(UserNumToLogin, { AccountId });
+	FTestPipeline& LoginPipeline = GetLoginPipeline(TestAccountIndex, { AccountId });
 
 	QueryFriendsHelperParams.OpParams->LocalAccountId = AccountId;
 	GetFriendsHelperParams.OpParams->LocalAccountId = AccountId;
@@ -89,14 +89,14 @@ SOCIAL_TEST_CASE("Verify that GetFriend returns a fail message if there are no c
 {
 	FAccountId AccountId;
 	
-	int32 UserNumToLogin = 5;
+	int32 TestAccountIndex = 5;
 
 	FGetFriends::Params OpGetParams;
 	FGetFriendsHelper::FHelperParams GetFriendsHelperParams;
 	GetFriendsHelperParams.OpParams = &OpGetParams;
 	GetFriendsHelperParams.ExpectedError = TOnlineResult<FGetFriends>(Errors::InvalidState());
 
-	FTestPipeline& LoginPipeline = GetLoginPipeline(UserNumToLogin, { AccountId });
+	FTestPipeline& LoginPipeline = GetLoginPipeline(TestAccountIndex, { AccountId });
 
 	GetFriendsHelperParams.OpParams->LocalAccountId = AccountId;
 
diff --git a/Engine/Source/Programs/Online/OnlineTestsCore/Source/Private/Tests/Social/QueryFriendsTests.cpp b/Engine/Source/Programs/Online/OnlineTestsCore/Source/Private/Tests/Social/QueryFriendsTests.cpp
--- a/Engine/Source/Programs/Online/OnlineTestsCore/Source/Private/Tests/Social/QueryFriendsTests.cpp
+++ b/Engine/Source/Programs/Online/OnlineTestsCore/Source/Private/Tests/Social/QueryFriendsTests.cpp
@@ -14,7 +14,7 @@ SOCIAL_TEST_CASE("Verify that QueryFriends returns a fail message if the local u
 {
 	FAccountId AccountId;
 
-	int32 UserNumToLogin = 1;
+	int32 TestAccountIndex = 1;
 	TSharedPtr<FPlatformUserId> AccountPlatformUserId = MakeShared<FPlatformUserId>();
 	bool bLogout = false;
 
@@ -23,7 +23,7 @@ SOCIAL_TEST_CASE("Verify that QueryFriends returns a fail message if the local u
 	QueryFriendsHelperParams.OpParams = &OpQueryParams;
 	QueryFriendsHelperParams.ExpectedError = TOnlineResult<FQueryFriends>(Errors::NotLoggedIn());
 
-	FTestPipeline& LoginPipeline = GetLoginPipeline(UserNumToLogin, { AccountId });
+	FTestPipeline& LoginPipeline = GetLoginPipeline(TestAccountIndex, { AccountId });
 
 	QueryFriendsHelperParams.OpParams->LocalAccountId = AccountId;
 
@@ -64,7 +64,7 @@ SOCIAL_TEST_CASE("Verify that QueryFriends returns a error if call with an inval
 SOCIAL_TEST_CASE("Verify that QueryFriends caches no Friends if no Friends exist for this user", EG_SOCIAL_QUERYFRIENDSEOS_TAG)
 {
 	FAccountId AccountId;
-	int32 UserNumToLogin = 5;
+	int32 TestAccountIndex = 5;
 
 	FQueryFriends::Params OpQueryParams;
 	FQueryFriendsHelper::FHelperParams QueryFriendsHelperParams;
@@ -74,7 +74,7 @@ SOCIAL_TEST_CASE("Verify that QueryFriends caches no Friends if no Friends exist
 	FGetFriendsHelper::FHelperParams GetFriendsHelperParams;
 	GetFriendsHelperParams.OpParams = &OpGetParams;
 
-	FTestPipeline& LoginPipeline = GetLoginPipeline(UserNumToLogin, { AccountId });
+	FTestPipeline& LoginPipeline = GetLoginPipeline(TestAccountIndex, { AccountId });
 
 	QueryFriendsHelperParams.OpParams->LocalAccountId = AccountId;
 	GetFriendsHelperParams.OpParams->LocalAccountId = AccountId;
@@ -89,7 +89,7 @@ SOCIAL_TEST_CASE("Verify that QueryFriends caches no Friends if no Friends exist
 SOCIAL_TEST_CASE("Verify that QueryFriends caches one Friend if only one Friend exists for this user", EG_SOCIAL_QUERYFRIENDSEOS_TAG)
 {
 	FAccountId AccountId;
-	int32 UserNumToLogin = 6;
+	int32 TestAccountIndex = 6;
 
 	FQueryFriends::Params OpQueryParams;
 	FQueryFriendsHelper::FHelperParams QueryFriendsHelperParams;
@@ -99,7 +99,7 @@ SOCIAL_TEST_CASE("Verify that QueryFriends caches one Friend if only one Friend
 	FGetFriendsHelper::FHelperParams GetFriendsHelperParams;
 	GetFriendsHelperParams.OpParams = &OpGetParams;
 		
-	FTestPipeline& LoginPipeline = GetLoginPipeline(UserNumToLogin, { AccountId });
+	FTestPipeline& LoginPipeline = GetLoginPipeline(TestAccountIndex, { AccountId });
 	
 	QueryFriendsHelperParams.OpParams->LocalAccountId = AccountId;
 	GetFriendsHelperParams.OpParams->LocalAccountId = AccountId;
diff --git a/Engine/Source/Programs/Online/OnlineTestsCore/Source/Private/Tests/Social/RejectFriendInviteTests.cpp b/Engine/Source/Programs/Online/OnlineTestsCore/Source/Private/Tests/Social/RejectFriendInviteTests.cpp
--- a/Engine/Source/Programs/Online/OnlineTestsCore/Source/Private/Tests/Social/RejectFriendInviteTests.cpp
+++ b/Engine/Source/Programs/Online/OnlineTestsCore/Source/Private/Tests/Social/RejectFriendInviteTests.cpp
@@ -47,7 +47,7 @@ SOCIAL_TEST_CASE("Verify that RejectFriendInvite returns a fail message if the l
 {
 	FAccountId FirstAccountId, SecondAccountId;
 
-	int32 UserNumToLogin = 1;
+	int32 TestAccountIndex = 1;
 	TSharedPtr<FPlatformUserId> FirstAccountPlatformUserId = MakeShared<FPlatformUserId>();
 	TSharedPtr<FPlatformUserId> SecondAccountPlatformUserId = MakeShared<FPlatformUserId>();
 	bool bLogout = false;
@@ -57,7 +57,7 @@ SOCIAL_TEST_CASE("Verify that RejectFriendInvite returns a fail message if the l
 	RejectFriendInviteHelperParams.OpParams = &OpRejectFriendInviteParams;
 	RejectFriendInviteHelperParams.ExpectedError = TOnlineResult<FRejectFriendInvite>(Errors::NotLoggedIn());
 
-	FTestPipeline& LoginPipeline = GetLoginPipeline(UserNumToLogin, { FirstAccountId, SecondAccountId });
+	FTestPipeline& LoginPipeline = GetLoginPipeline(TestAccountIndex, { FirstAccountId, SecondAccountId });
 
 	RejectFriendInviteHelperParams.OpParams->LocalAccountId = FirstAccountId;
 	RejectFriendInviteHelperParams.OpParams->TargetAccountId = SecondAccountId;
diff --git a/Engine/Source/Programs/Online/OnlineTestsCore/Source/Private/Tests/Social/SendFriendInviteTests.cpp b/Engine/Source/Programs/Online/OnlineTestsCore/Source/Private/Tests/Social/SendFriendInviteTests.cpp
--- a/Engine/Source/Programs/Online/OnlineTestsCore/Source/Private/Tests/Social/SendFriendInviteTests.cpp
+++ b/Engine/Source/Programs/Online/OnlineTestsCore/Source/Private/Tests/Social/SendFriendInviteTests.cpp
@@ -51,7 +51,7 @@ SOCIAL_TEST_CASE("Verify that SendFriendInvite returns a fail message if the loc
 {
 	FAccountId FirstAccountId, SecondAccountId;
 
-	int32 UserNumToLogin = 1;
+	int32 TestAccountIndex = 1;
 	TSharedPtr<FPlatformUserId> FirstAccountPlatformUserId = MakeShared<FPlatformUserId>();
 	TSharedPtr<FPlatformUserId> SecondAccountPlatformUserId = MakeShared<FPlatformUserId>();
 	bool bLogout = false;
@@ -61,7 +61,7 @@ SOCIAL_TEST_CASE("Verify that SendFriendInvite returns a fail message if the loc
 	SendFriendInviteHelperParams.OpParams = &OpSendFriendInviteParams;
 	SendFriendInviteHelperParams.ExpectedError = TOnlineResult<FSendFriendInvite>(Errors::NotLoggedIn());
 
-	FTestPipeline& LoginPipeline = GetLoginPipeline(UserNumToLogin, { FirstAccountId, SecondAccountId });
+	FTestPipeline& LoginPipeline = GetLoginPipeline(TestAccountIndex, { FirstAccountId, SecondAccountId });
 
 	SendFriendInviteHelperParams.OpParams->LocalAccountId = FirstAccountId;
 	SendFriendInviteHelperParams.OpParams->TargetAccountId = SecondAccountId;
diff --git a/Engine/Source/Programs/Online/OnlineTestsCore/Source/Public/OnlineCatchHelper.h b/Engine/Source/Programs/Online/OnlineTestsCore/Source/Public/OnlineCatchHelper.h
--- a/Engine/Source/Programs/Online/OnlineTestsCore/Source/Public/OnlineCatchHelper.h
+++ b/Engine/Source/Programs/Online/OnlineTestsCore/Source/Public/OnlineCatchHelper.h
@@ -39,7 +39,7 @@ public:
 
 #if ONLINETESTS_USEEXTERNAUTH
 	/* Returning accounts to the pool */
-	static bool ReturnAccounts(int32 TemplateRequestedNum = -1);
+	static bool ReturnAccounts(int32 TestAccountIndex = -1);
 #endif
 
 	FTestPipeline& GetLoginPipeline(std::initializer_list<std::reference_wrapper<FAccountId>> AccountIds) const;
@@ -53,7 +53,7 @@ public:
 	void AssignLoginUsers(int32 LocalUserId, FAccountId& OutAccountId) const;
 
 #if ONLINETESTS_USEEXTERNAUTH
-	TArray<FString> GetExternalAuthAccountIds(int32 TemplateRequestedNum) const;
+	TArray<FString> GetExternalAuthAccountIds(int32 TestAccountIndex) const;
 #endif
 
 	FTestPipeline& GetPipeline() const;
@@ -62,13 +62,13 @@ protected:
 	OnlineTestBase();
 
 	/* Helper function to delete accounts from TestDataService */
-	bool DeleteAccounts(int32 TemplateRequestedNum = -1) const;
+	bool DeleteAccounts(int32 TestAccountIndex = -1) const;
 
 	/* Helper function to destroy the current OnlineService module which stores some state that maybe needs to be reset */
 	void DestroyCurrentServiceModule() const;
 
 	/* Proxy function to the DestroyCurrentServiceModule() */
-	bool ResetAccountStatus(int32 TemplateRequestedNum = -1) const;
+	bool ResetAccountStatus(int32 TestAccountIndex = -1) const;
 
 	/* Returns the name of the service we're currently testing */
 	FString GetService() const;
@@ -82,17 +82,17 @@ protected:
 
 	// When called should "reset" accounts to default state for Stats, Achivements, Lobbies and other
 	// stateful interfaces to work.
-	bool CustomResetAccounts(int32 TemplateRequestedNum) const;
+	bool CustomResetAccounts(int32 TestAccountIndex) const;
 
 	// When called will *remove* the account from the underlying service.
 	// This is generally called to test new-user login or first-time setup.
 	// To be called rare-ly.
-	bool CustomDeleteAccounts(int32 TemplateRequestedNum) const;
+	bool CustomDeleteAccounts(int32 TestAccountIndex) const;
 
 #endif
-	TArray<UE::Online::FAuthLogin::Params> GetIniCredentials(int32 LocalUserNum) const;
+	TArray<UE::Online::FAuthLogin::Params> GetIniCredentials(int32 TestAccountIndex) const;
 
-	TArray<UE::Online::FAuthLogin::Params> GetCredentials(int32 LocalUserNum, int32 NumUsers) const;
+	TArray<UE::Online::FAuthLogin::Params> GetCredentials(int32 TestAccountIndex, int32 NumUsers) const;
 
 	/* Returns the ini login category name for the configured service */
 	FString GetLoginCredentialCategory() const;
@@ -110,7 +110,7 @@ private:
 	mutable FTestDriver Driver;
 	mutable TSharedPtr<FTestPipeline> Pipeline;
 	mutable uint32 NumLocalUsers = -1;
-	mutable uint32 UserNumToLogout = -1;
+	mutable uint32 TestAccountID = -1;
 };
 
 typedef OnlineTestBase* (*OnlineTestConstructor)();
