Fix errors and warnings in OSS and Online:

- fix falls errors in Achievements suite tests
- fix warning from SDK

diff --git a/Engine/Source/Programs/Online/OSSTestsCore/OnlineSubsystemCatchHelper.cpp b/Engine/Source/Programs/Online/OSSTestsCore/OnlineSubsystemCatchHelper.cpp
--- a/Engine/Source/Programs/Online/OSSTestsCore/OnlineSubsystemCatchHelper.cpp
+++ b/Engine/Source/Programs/Online/OSSTestsCore/OnlineSubsystemCatchHelper.cpp
@@ -9,6 +9,9 @@
 #include "Helpers/Identity/IdentityAutoLoginHelper.h"
 #include "Helpers/Identity/IdentityLoginHelper.h"
 #include "Helpers/Identity/IdentityLogoutHelper.h"
+
+#include "Helpers/TickForTime.h"
+
 #include "OnlineSubsystemNames.h"
 #include "Misc/CommandLine.h"
 #include "GenericPlatform/GenericPlatformInputDeviceMapper.h"
@@ -190,7 +193,7 @@ FTestPipeline& FOnlineSubsystemTestBaseFixture::GetPipeline() const
 	return GetLoginPipeline(0);
 }
 
-void FOnlineSubsystemTestBaseFixture::RunToCompletion() const
+void FOnlineSubsystemTestBaseFixture::RunToCompletion(bool bWaitBeforeLogout, const FTimespan TimeToWaitMilliseconds) const
 {
 	bool bUseAutoLogin = false;
 	bool bUseImplicitLogin = false;
@@ -211,6 +214,12 @@ void FOnlineSubsystemTestBaseFixture::RunToCompletion() const
 	{
 		for (uint32 i = 0; i < NumLocalUsers; i++)
 		{
+			if (bWaitBeforeLogout)
+			{
+				Pipeline
+					.EmplaceStep<FTickForTime>(TimeToWaitMilliseconds);
+			}
+
 			Pipeline.EmplaceStep<FIdentityLogoutStep>(i);
 		}
 	}
@@ -222,7 +231,8 @@ void FOnlineSubsystemTestBaseFixture::RunToCompletion() const
 	Driver.RunToCompletion();
 }
 
-void FOnlineSubsystemFixtureInvoker::SetSubsystem(const FString& InStoredSubsystem) {
+void FOnlineSubsystemFixtureInvoker::SetSubsystem(const FString& InStoredSubsystem) 
+{
 	StoredSubsystem = InStoredSubsystem;
 }
 
diff --git a/Engine/Source/Programs/Online/OSSTestsCore/OnlineSubsystemCatchHelper.h b/Engine/Source/Programs/Online/OSSTestsCore/OnlineSubsystemCatchHelper.h
--- a/Engine/Source/Programs/Online/OSSTestsCore/OnlineSubsystemCatchHelper.h
+++ b/Engine/Source/Programs/Online/OSSTestsCore/OnlineSubsystemCatchHelper.h
@@ -80,19 +80,24 @@ template <typename C>
 class TestInvokerFixture : public FOnlineSubsystemFixtureInvoker
 {
 public:
-	constexpr TestInvokerFixture(void (C::* InTestAsMethod)()) noexcept :
-		TestAsMethod(InTestAsMethod) {}
+	constexpr TestInvokerFixture(void (C::* InTestAsMethod)()) noexcept 
+		: TestAsMethod(InTestAsMethod) 
+	{
+	}
 
-	void PrepareTestCase() const {
+	void PrepareTestCase() const
+	{
 		Fixture = Catch::Detail::make_unique<C>();
 		Fixture->ConstructInternal(StoredSubsystem);
 	}
 
-	void TearDownTestCase() const {
+	void TearDownTestCase() const
+	{
 		Fixture.reset();
 	}
 
-	void invoke() const override {
+	void invoke() const override
+	{
 		PrepareTestCase();
 		auto* f = Fixture.get();
 		(f->*TestAsMethod)();
@@ -143,7 +148,7 @@ protected:
 	/* Returns the ini login category name for the configured service */
 	FString GetLoginCredentialCategory() const;
 
-	void RunToCompletion() const;
+	void RunToCompletion(bool bWaitBeforeLogout = false, const FTimespan TimeToWaitMilliseconds = FTimespan::FromMilliseconds(1000)) const;
 
 	/* ITestInvoker */
 	virtual void invoke() const override = 0;
diff --git a/Engine/Source/Programs/Online/OSSTestsCore/Tests/Achievements/AchievementsGetCachedAchievementDescriptionTests.cpp b/Engine/Source/Programs/Online/OSSTestsCore/Tests/Achievements/AchievementsGetCachedAchievementDescriptionTests.cpp
--- a/Engine/Source/Programs/Online/OSSTestsCore/Tests/Achievements/AchievementsGetCachedAchievementDescriptionTests.cpp
+++ b/Engine/Source/Programs/Online/OSSTestsCore/Tests/Achievements/AchievementsGetCachedAchievementDescriptionTests.cpp
@@ -9,6 +9,8 @@
 #include "Helpers/Achievements/AchievementsQueryAchievementDescriptionsHelper.h"
 #include "Helpers/Achievements/AchievementsGetCachedAchievementDescriptionHelper.h"
 
+#include "Helpers/OnlineSubsystemTestFixtures.h"
+
 #define ACHIEVEMENTS_TAG "[suite_achievements]"
 #define EG_ACHIEVEMENTS_GETCACHEDACHIEMENTDESCRIPTION_TAG ACHIEVEMENTS_TAG "[getcachedachievementdescription]"
 
@@ -18,13 +20,14 @@ ACHIEVEMENTS_TEST_CASE("Verify calling Achievements GetCachedAchievementDescript
 {
 	int32 LocalUserNum = 0;
 	FUniqueNetIdPtr LocalUserId = nullptr;
-	const FString AchievementId = TEXT("null-ach-0");
+	const FString AchievementId = TEXT("test_getachievementdefinitioncount");
 	int32 NumUsersToImplicitLogin = 1;
+	bool bWaitBeforeLogout = true;
 
 	GetLoginPipeline(NumUsersToImplicitLogin)
 		.EmplaceStep<FIdentityGetUniquePlayerIdStep>(LocalUserNum, [&LocalUserId](FUniqueNetIdPtr InUserId) {LocalUserId = InUserId; })
 		.EmplaceStep<FAchievementsQueryAchievementDescriptionsStep>(&LocalUserId)
 		.EmplaceStep<FAchievementsGetCachedAchievementDescriptionStep>(AchievementId);
 	
-	RunToCompletion();
+	RunToCompletion(bWaitBeforeLogout);
 }
diff --git a/Engine/Source/Programs/Online/OSSTestsCore/Tests/Achievements/AchievementsGetCachedAchievementTests.cpp b/Engine/Source/Programs/Online/OSSTestsCore/Tests/Achievements/AchievementsGetCachedAchievementTests.cpp
--- a/Engine/Source/Programs/Online/OSSTestsCore/Tests/Achievements/AchievementsGetCachedAchievementTests.cpp
+++ b/Engine/Source/Programs/Online/OSSTestsCore/Tests/Achievements/AchievementsGetCachedAchievementTests.cpp
@@ -19,13 +19,14 @@ ACHIEVEMENTS_TEST_CASE("Verify calling Achievements GetCachedAchievement with va
 	int32 LocalUserNum = 0;
 	FUniqueNetIdPtr LocalUserId = nullptr;
 	FOnlineAchievement* Achievenment = nullptr;
-	const FString AchievementId = TEXT("null-ach-0");
+	const FString AchievementId = TEXT("test_getachievementdefinitioncount");
 	int32 NumUsersToImplicitLogin = 1;
+	bool bWaitBeforeLogout = true;
 
 	GetLoginPipeline(NumUsersToImplicitLogin)
 		.EmplaceStep<FIdentityGetUniquePlayerIdStep>(LocalUserNum, [&LocalUserId](FUniqueNetIdPtr InUserId) {LocalUserId = InUserId; })
 		.EmplaceStep<FAchievementsQueryAchievementsStep>(&LocalUserId)
 		.EmplaceStep<FAchievementsGetCachedAchievementStep>(&LocalUserId, AchievementId, [&Achievenment](FOnlineAchievement* InOnlineAchievement) {Achievenment = InOnlineAchievement; });
 	
-	RunToCompletion();
+	RunToCompletion(bWaitBeforeLogout);
 }
diff --git a/Engine/Source/Programs/Online/OSSTestsCore/Tests/Achievements/AchievementsGetCachedAchievementsTests.cpp b/Engine/Source/Programs/Online/OSSTestsCore/Tests/Achievements/AchievementsGetCachedAchievementsTests.cpp
--- a/Engine/Source/Programs/Online/OSSTestsCore/Tests/Achievements/AchievementsGetCachedAchievementsTests.cpp
+++ b/Engine/Source/Programs/Online/OSSTestsCore/Tests/Achievements/AchievementsGetCachedAchievementsTests.cpp
@@ -20,11 +20,12 @@ ACHIEVEMENTS_TEST_CASE("Verify calling Achievements GetCachedAchievements with v
 	FUniqueNetIdPtr LocalUserId = nullptr;
 	TArray<FOnlineAchievement>* Achievenments = nullptr;
 	int32 NumUsersToImplicitLogin = 1;
+	bool bWaitBeforeLogout = true;
 
 	GetLoginPipeline(NumUsersToImplicitLogin)
 		.EmplaceStep<FIdentityGetUniquePlayerIdStep>(LocalUserNum, [&LocalUserId](FUniqueNetIdPtr InUserId) {LocalUserId = InUserId; })
 		.EmplaceStep<FAchievementsQueryAchievementsStep>(&LocalUserId)
 		.EmplaceStep<FAchievementsGetCachedAchievementsStep>(&LocalUserId, [&Achievenments](TArray<FOnlineAchievement>* InOnlineAchievements) {Achievenments = InOnlineAchievements; });
 	
-	RunToCompletion();
+	RunToCompletion(bWaitBeforeLogout);
 }
diff --git a/Engine/Source/Programs/Online/OSSTestsCore/Tests/Achievements/AchievementsQueryAchievementDescriptionsTests.cpp b/Engine/Source/Programs/Online/OSSTestsCore/Tests/Achievements/AchievementsQueryAchievementDescriptionsTests.cpp
--- a/Engine/Source/Programs/Online/OSSTestsCore/Tests/Achievements/AchievementsQueryAchievementDescriptionsTests.cpp
+++ b/Engine/Source/Programs/Online/OSSTestsCore/Tests/Achievements/AchievementsQueryAchievementDescriptionsTests.cpp
@@ -18,10 +18,11 @@ ACHIEVEMENTS_TEST_CASE("Verify calling Achievements QueryAchievementDescriptions
 	int32 LocalUserNum = 0;
 	FUniqueNetIdPtr LocalUserId = nullptr;
 	int32 NumUsersToImplicitLogin = 1;
+	bool bWaitBeforeLogout = true;
 
 	GetLoginPipeline(NumUsersToImplicitLogin)
 		.EmplaceStep<FIdentityGetUniquePlayerIdStep>(LocalUserNum, [&LocalUserId](FUniqueNetIdPtr InUserId) {LocalUserId = InUserId; })
 		.EmplaceStep<FAchievementsQueryAchievementDescriptionsStep>(&LocalUserId);
 
-	RunToCompletion();
+	RunToCompletion(bWaitBeforeLogout);
 }
diff --git a/Engine/Source/Programs/Online/OSSTestsCore/Tests/Achievements/AchievementsQueryAchievementsTests.cpp b/Engine/Source/Programs/Online/OSSTestsCore/Tests/Achievements/AchievementsQueryAchievementsTests.cpp
--- a/Engine/Source/Programs/Online/OSSTestsCore/Tests/Achievements/AchievementsQueryAchievementsTests.cpp
+++ b/Engine/Source/Programs/Online/OSSTestsCore/Tests/Achievements/AchievementsQueryAchievementsTests.cpp
@@ -18,10 +18,11 @@ ACHIEVEMENTS_TEST_CASE("Verify calling Achievements QueryAchievements with valid
 	int32 LocalUserNum = 0;
 	FUniqueNetIdPtr LocalUserId = nullptr;
 	int32 NumUsersToImplicitLogin = 1;
+	bool bWaitBeforeLogout = true;
 
 	GetLoginPipeline(NumUsersToImplicitLogin)
 		.EmplaceStep<FIdentityGetUniquePlayerIdStep>(LocalUserNum, [&LocalUserId](FUniqueNetIdPtr InUserId) {LocalUserId = InUserId; })
 		.EmplaceStep<FAchievementsQueryAchievementsStep>(&LocalUserId);
 
-	RunToCompletion();
+	RunToCompletion(bWaitBeforeLogout);
 }
diff --git a/Engine/Source/Programs/Online/OSSTestsCore/Tests/Achievements/AchievementsWriteAchievementsTests.cpp b/Engine/Source/Programs/Online/OSSTestsCore/Tests/Achievements/AchievementsWriteAchievementsTests.cpp
--- a/Engine/Source/Programs/Online/OSSTestsCore/Tests/Achievements/AchievementsWriteAchievementsTests.cpp
+++ b/Engine/Source/Programs/Online/OSSTestsCore/Tests/Achievements/AchievementsWriteAchievementsTests.cpp
@@ -20,12 +20,19 @@ ACHIEVEMENTS_TEST_CASE("Verify calling Achievements WriteAchievements with valid
 	FUniqueNetIdPtr LocalUserId = nullptr;
 	FOnlineAchievementsWritePtr AchievementWriteObject = MakeShareable(new FOnlineAchievementsWrite());
 	FOnlineAchievementsWriteRef AchievementWriter = AchievementWriteObject.ToSharedRef();
+
+	double CurrentTimeSeconds = FPlatformTime::Seconds();
+	FString TimestampString = FString::Printf(TEXT("%f"), CurrentTimeSeconds);
+
+	AchievementWriter->Properties.Emplace(/*TimestampString*/"test_unlockachievements", 1);
+
 	int32 NumUsersToImplicitLogin = 1;
+	bool bWaitBeforeLogout = true;
 
 	GetLoginPipeline(NumUsersToImplicitLogin)
 		.EmplaceStep<FIdentityGetUniquePlayerIdStep>(LocalUserNum, [&LocalUserId](FUniqueNetIdPtr InUserId) {LocalUserId = InUserId; })
 		.EmplaceStep<FAchievementsQueryAchievementsStep>(&LocalUserId)
 		.EmplaceStep<FAchievementsWriteAchievementsStep>(&LocalUserId, AchievementWriter);
 
-	RunToCompletion();
+	RunToCompletion(bWaitBeforeLogout);
 }
