Redo the use of smart pointers when registering tests in OSS Tests

- tests for different OSS should be run in terms of previous changes for fixtures properly
- allowed to remove the number of errors in Achievements test suite

diff --git a/Engine/Source/Programs/Online/OSSTestsCore/OnlineSubsystemCatchHelper.cpp b/Engine/Source/Programs/Online/OSSTestsCore/OnlineSubsystemCatchHelper.cpp
--- a/Engine/Source/Programs/Online/OSSTestsCore/OnlineSubsystemCatchHelper.cpp
+++ b/Engine/Source/Programs/Online/OSSTestsCore/OnlineSubsystemCatchHelper.cpp
@@ -523,9 +523,11 @@ FOnlineSubsystemAutoReg::FOnlineSubsystemAutoReg(Catch::Detail::unique_ptr<FOnli
 				auto GeneratedTags = StringCast<ANSICHAR>(*FOnlineSubsystemFixtureInvoker::GenerateTags(SubsystemName, SkippableTags, ANSI_TO_TCHAR(Tags)));
 				FString NewTags = GeneratedTags.Get();
 
+				auto ClonedTestInvokerPtr = TestInvokerPtr->clone();
+
 				if (!FOnlineSubsystemFixtureInvoker::IsRunningTestSkipOnTags(SubsystemName, SkippableTags, NewTags))
 				{
-					TestInvokerPtr->SetSubsystem(Config.Name);
+					ClonedTestInvokerPtr->SetSubsystem(SubsystemName);
 				}
 				else
 				{
@@ -536,9 +538,12 @@ FOnlineSubsystemAutoReg::FOnlineSubsystemAutoReg(Catch::Detail::unique_ptr<FOnli
 					std::string(Catch::StringRef()),  // Used for testing a static method instead of a function- not needed since we're passing an ITestInvoker macro
 					Catch::NameAndTags{ NewName.Get(), TCHAR_TO_ANSI(*NewTags) },
 					LineInfo),
-					Catch::Detail::unique_ptr(TestInvokerPtr) // This is taking the ITestInvoker macro and will call invoke() to run the test
+					Catch::Detail::unique_ptr(ClonedTestInvokerPtr.release()) // This is taking the ITestInvoker macro and will call invoke() to run the test
 				);
+
 			}
+
+			delete TestInvokerPtr;
 		});
 }
 
diff --git a/Engine/Source/Programs/Online/OSSTestsCore/OnlineSubsystemCatchHelper.h b/Engine/Source/Programs/Online/OSSTestsCore/OnlineSubsystemCatchHelper.h
--- a/Engine/Source/Programs/Online/OSSTestsCore/OnlineSubsystemCatchHelper.h
+++ b/Engine/Source/Programs/Online/OSSTestsCore/OnlineSubsystemCatchHelper.h
@@ -40,6 +40,8 @@ public:
 
 	virtual void invoke() const = 0;
 
+	virtual Catch::Detail::unique_ptr<FOnlineSubsystemFixtureInvoker> clone() = 0;
+
 	void SetSubsystem(const FString& InStoredSubsystem);
 
 	// TODO: This should poll some info from the tags and generate the list based on that
@@ -104,6 +106,11 @@ public:
 		TearDownTestCase();
 	}
 
+	Catch::Detail::unique_ptr<FOnlineSubsystemFixtureInvoker> clone() override
+	{
+		return Catch::Detail::unique_ptr<FOnlineSubsystemFixtureInvoker>(Catch::Detail::make_unique<TestInvokerFixture<C>>(TestAsMethod));
+	}
+
 private:
 	void (C::* TestAsMethod)();
 	mutable Catch::Detail::unique_ptr<C> Fixture = nullptr;
diff --git a/Engine/Source/Programs/Online/OSSTestsCore/Tests/AccountServiceTest.cpp b/Engine/Source/Programs/Online/OSSTestsCore/Tests/AccountServiceTest.cpp
--- a/Engine/Source/Programs/Online/OSSTestsCore/Tests/AccountServiceTest.cpp
+++ b/Engine/Source/Programs/Online/OSSTestsCore/Tests/AccountServiceTest.cpp
@@ -7,6 +7,9 @@
 
 #include "OnlineSubsystemCatchHelper.h"
 
+#include "Helpers/TickForTime.h"
+#include "Helpers/Identity/IdentityLogoutHelper.h"
+
 #define ACCOUNTSERVICE_TAG "[AccountService]"
 
 #define ACCOUNTSERVICE_TEST_CASE(x, ...) ONLINESUBSYSTEM_TEST_CASE(x, ACCOUNTSERVICE_TAG __VA_ARGS__)
@@ -23,17 +26,20 @@ ACCOUNTSERVICE_TEST_CASE("Verify if we can properly create a OnlineAccountCreden
 
 ACCOUNTSERVICE_TEST_CASE("Verify if we can properly instantiate the OSS")
 {
+	int32 LocalUserNum = 0;
+	int32 NumUsers = 1;
+
 	GetPipeline()
-		.EmplaceLambda([this](IOnlineSubsystem* Services)
+		.EmplaceLambda([this, &LocalUserNum, &NumUsers](IOnlineSubsystem* Services)
 		{
-			int32 LocalUserNum = 0;
-			int32 NumUsers = 1;
 			TArray<FOnlineAccountCredentials> AccountCreds = GetCredentials(LocalUserNum, NumUsers);
 
 			IOnlineIdentityPtr IdentityInterface = Services->GetIdentityInterface();		
 			bool LoggedIn = IdentityInterface.Get()->Login(LocalUserNum, AccountCreds[0]);
 			REQUIRE(LoggedIn);
-		});
+		})
+		.EmplaceStep<FTickForTime>(FTimespan::FromMilliseconds(1000))
+		.EmplaceStep<FIdentityLogoutStep>(LocalUserNum);
 
 	RunToCompletion();
 }
\ No newline at end of file
diff --git a/Engine/Source/Programs/Online/OSSTestsCore/Tests/Achievements/AchievementsGetCachedAchievementDescriptionTests.cpp b/Engine/Source/Programs/Online/OSSTestsCore/Tests/Achievements/AchievementsGetCachedAchievementDescriptionTests.cpp
--- a/Engine/Source/Programs/Online/OSSTestsCore/Tests/Achievements/AchievementsGetCachedAchievementDescriptionTests.cpp
+++ b/Engine/Source/Programs/Online/OSSTestsCore/Tests/Achievements/AchievementsGetCachedAchievementDescriptionTests.cpp
@@ -9,6 +9,8 @@
 #include "Helpers/Achievements/AchievementsQueryAchievementDescriptionsHelper.h"
 #include "Helpers/Achievements/AchievementsGetCachedAchievementDescriptionHelper.h"
 
+#include "Helpers/TickForTime.h"
+
 #include "Helpers/OnlineSubsystemTestFixtures.h"
 
 #define ACHIEVEMENTS_TAG "[suite_achievements]"
@@ -20,14 +22,24 @@ ACHIEVEMENTS_TEST_CASE("Verify calling Achievements GetCachedAchievementDescript
 {
 	int32 LocalUserNum = 0;
 	FUniqueNetIdPtr LocalUserId = nullptr;
-	const FString AchievementId = TEXT("test_getachievementdefinitioncount");
+	FString AchievementId = TEXT("");
 	int32 NumUsersToImplicitLogin = 1;
 	bool bWaitBeforeLogout = true;
 
+	if (GetSubsystem() == "EOS")
+	{
+		AchievementId = TEXT("test_getachievementdefinitioncount");
+	}
+	else if(GetSubsystem() == "NULL")
+	{
+		AchievementId = TEXT("null-ach-0");
+	}
+
 	GetLoginPipeline(NumUsersToImplicitLogin)
 		.EmplaceStep<FIdentityGetUniquePlayerIdStep>(LocalUserNum, [&LocalUserId](FUniqueNetIdPtr InUserId) {LocalUserId = InUserId; })
 		.EmplaceStep<FAchievementsQueryAchievementDescriptionsStep>(&LocalUserId)
+		.EmplaceStep<FTickForTime>(FTimespan::FromMilliseconds(1000))
 		.EmplaceStep<FAchievementsGetCachedAchievementDescriptionStep>(AchievementId);
-	
+		
 	RunToCompletion(bWaitBeforeLogout);
 }
diff --git a/Engine/Source/Programs/Online/OSSTestsCore/Tests/Achievements/AchievementsGetCachedAchievementTests.cpp b/Engine/Source/Programs/Online/OSSTestsCore/Tests/Achievements/AchievementsGetCachedAchievementTests.cpp
--- a/Engine/Source/Programs/Online/OSSTestsCore/Tests/Achievements/AchievementsGetCachedAchievementTests.cpp
+++ b/Engine/Source/Programs/Online/OSSTestsCore/Tests/Achievements/AchievementsGetCachedAchievementTests.cpp
@@ -19,10 +19,19 @@ ACHIEVEMENTS_TEST_CASE("Verify calling Achievements GetCachedAchievement with va
 	int32 LocalUserNum = 0;
 	FUniqueNetIdPtr LocalUserId = nullptr;
 	FOnlineAchievement* Achievenment = nullptr;
-	const FString AchievementId = TEXT("test_getachievementdefinitioncount");
+	FString AchievementId = TEXT("");
 	int32 NumUsersToImplicitLogin = 1;
 	bool bWaitBeforeLogout = true;
 
+	if (GetSubsystem() == "EOS")
+	{
+		AchievementId = TEXT("test_getachievementdefinitioncount");
+	}
+	else if (GetSubsystem() == "NULL")
+	{
+		AchievementId = TEXT("null-ach-0");
+	}
+
 	GetLoginPipeline(NumUsersToImplicitLogin)
 		.EmplaceStep<FIdentityGetUniquePlayerIdStep>(LocalUserNum, [&LocalUserId](FUniqueNetIdPtr InUserId) {LocalUserId = InUserId; })
 		.EmplaceStep<FAchievementsQueryAchievementsStep>(&LocalUserId)
diff --git a/Engine/Source/Programs/Online/OSSTestsCore/Tests/Achievements/AchievementsWriteAchievementsTests.cpp b/Engine/Source/Programs/Online/OSSTestsCore/Tests/Achievements/AchievementsWriteAchievementsTests.cpp
--- a/Engine/Source/Programs/Online/OSSTestsCore/Tests/Achievements/AchievementsWriteAchievementsTests.cpp
+++ b/Engine/Source/Programs/Online/OSSTestsCore/Tests/Achievements/AchievementsWriteAchievementsTests.cpp
@@ -24,7 +24,7 @@ ACHIEVEMENTS_TEST_CASE("Verify calling Achievements WriteAchievements with valid
 	double CurrentTimeSeconds = FPlatformTime::Seconds();
 	FString TimestampString = FString::Printf(TEXT("%f"), CurrentTimeSeconds);
 
-	AchievementWriter->Properties.Emplace(/*TimestampString*/"test_unlockachievements", 1);
+	AchievementWriter->Properties.Emplace("test_unlockachievements", 1);
 
 	int32 NumUsersToImplicitLogin = 1;
 	bool bWaitBeforeLogout = true;
