Fix configuration and update template for OSMcp low level tests
Based on Stephen.Ma CL 40186277
#rnx

diff --git a/Engine/Source/Programs/Online/OSSTestsCore/Tests/Session/SessionFindFriendSessionTests.cpp b/Engine/Source/Programs/Online/OSSTestsCore/Tests/Session/SessionFindFriendSessionTests.cpp
--- a/Engine/Source/Programs/Online/OSSTestsCore/Tests/Session/SessionFindFriendSessionTests.cpp
+++ b/Engine/Source/Programs/Online/OSSTestsCore/Tests/Session/SessionFindFriendSessionTests.cpp
@@ -28,7 +28,7 @@ SESSION_TEST_CASE("Verify calling Session FindFriendSession with LocalUserId and
 	FUniqueNetIdPtr LocalUserId = nullptr;
 	FUniqueNetIdPtr TargetUserId = nullptr;
 	FName SessionName = FName(FString::Printf(TEXT("TestSession_%s"), *FTestUtilities::GetUniqueTestString()));
-	FOnlineSessionSetting GamemodeSetting(TEXT("FORTEMPTY"), EOnlineDataAdvertisementType::ViaOnlineService);
+	FOnlineSessionSetting GamemodeSetting(TEXT("FORTEMPTY"), EOnlineDataAdvertisementType::ViaOnlineServiceAndPing);
 
 	FOnlineSessionSettings SessionSettings;
 	SessionSettings.bShouldAdvertise = true;
@@ -45,7 +45,7 @@ SESSION_TEST_CASE("Verify calling Session FindFriendSession with LocalUserId and
 		.EmplaceStep<FIdentityGetUniquePlayerIdStep>(TargetUserNum, [&TargetUserId](FUniqueNetIdPtr InUserId) {TargetUserId = InUserId; })
 		.EmplaceStep<FFriendsEnsureFriendshipStep>(LocalUserNum, TargetUserNum, &LocalUserId, &TargetUserId, LocalListName, bLocalIsFriendsListPopulated)
 		.EmplaceStep<FSessionCreateSessionStep>(&LocalUserId, SessionName, SessionSettings)
-		.EmplaceStep<FSessionFindFriendSessionStep>(&TargetUserId, TargetUserNum, &LocalUserId)
+		.EmplaceStep<FSessionFindFriendSessionStep>(&LocalUserId, LocalUserNum, &TargetUserId)
 		.EmplaceStep<FSessionDestroySessionStep>(SessionName);
 
 	RunToCompletion();
