Naming correctness for online tests

diff --git a/Engine/Source/Programs/Online/OnlineTestsCore/Source/Private/OnlineCatchHelper.cpp b/Engine/Source/Programs/Online/OnlineTestsCore/Source/Private/OnlineCatchHelper.cpp
--- a/Engine/Source/Programs/Online/OnlineTestsCore/Source/Private/OnlineCatchHelper.cpp
+++ b/Engine/Source/Programs/Online/OnlineTestsCore/Source/Private/OnlineCatchHelper.cpp
@@ -73,7 +73,7 @@ EOnlineServices OnlineTestBase::GetServiceType() const
 	return ServiceType;
 }
 
-SubsystemType OnlineTestBase::GetSubsystem() const
+IOnlineServicesPtr OnlineTestBase::GetSubsystem() const
 {
 	return UE::Online::FOnlineServicesRegistry::Get().GetNamedServicesInstance(ServiceType, NAME_None, NAME_None);
 }
@@ -189,8 +189,8 @@ FString OnlineTestBase::GetLoginCredentialCategory() const
 
 void OnlineTestBase::AssignLoginUsers(int32 LocalUserId, FAccountId& OutAccountId) const
 {
-	SubsystemType OnlineSubsystem = GetSubsystem();
-	UE::Online::TOnlineResult<UE::Online::FAuthGetLocalOnlineUserByPlatformUserId> UserId = OnlineSubsystem->GetAuthInterface()->GetLocalOnlineUserByPlatformUserId({ FPlatformMisc::GetPlatformUserForUserIndex(LocalUserId) });
+	IOnlineServicesPtr OnlineServices = GetSubsystem();
+	UE::Online::TOnlineResult<UE::Online::FAuthGetLocalOnlineUserByPlatformUserId> UserId = OnlineServices->GetAuthInterface()->GetLocalOnlineUserByPlatformUserId({ FPlatformMisc::GetPlatformUserForUserIndex(LocalUserId) });
 	REQUIRE(UserId.IsOk());
 	CHECK(UserId.TryGetOkValue() != nullptr);
 	OutAccountId = UserId.TryGetOkValue()->AccountInfo->AccountId;
diff --git a/Engine/Source/Programs/Online/OnlineTestsCore/Source/Private/TestDriver.cpp b/Engine/Source/Programs/Online/OnlineTestsCore/Source/Private/TestDriver.cpp
--- a/Engine/Source/Programs/Online/OnlineTestsCore/Source/Private/TestDriver.cpp
+++ b/Engine/Source/Programs/Online/OnlineTestsCore/Source/Private/TestDriver.cpp
@@ -53,7 +53,7 @@ FString FTestDriver::TimeoutFailedTestInfo() const
 	return FailedInfoString;
 }
 
-void FTestPipeline::operator()(SubsystemType Subsystem)
+void FTestPipeline::operator()(const IOnlineServicesPtr& Subsystem)
 {
 	double CurrentTime = FPlatformTime::Seconds();
 
@@ -140,19 +140,19 @@ void FTestPipeline::EvaluatePlatformTickTime(double&& TickTime)
 		{
 			FTimespan AverageTick = SubsystemTickSum.GetTicks() / SubsystemTickCount;
 			FTimespan ExpectedAverageTick = EvaluateTickConfig.ExpectedAverageTick.GetTicks();
-			//CHECK(AverageTick <= ExpectedAverageTick); // todo: what are these? theyre spamming checks in my code
+			//CHECK(AverageTick <= ExpectedAverageTick); // todo: what are these? they are spamming checks in my code
 		}
 
 		FTimespan Tick = FTimespan::FromSeconds(TickTime);
 		FTimespan AbsoluteMaximumTick = EvaluateTickConfig.AbsoluteMaximumTick.GetTicks();
-		//CHECK(Tick <= AbsoluteMaximumTick); // todo: what are these? theyre spamming checks in my code
+		//CHECK(Tick <= AbsoluteMaximumTick); // todo: what are these? they are spamming checks in my code
 	}
 }
 
 FTestDriver::~FTestDriver()
 {
 	// Mark all remaining subsystems as complete.
-	ForeachSubsystemInstance([this](SubsystemType OnlineSubsystem, FTestPipeline&)
+	ForeachSubsystemInstance([this](const IOnlineServicesPtr&  OnlineSubsystem, FTestPipeline&)
 		{
 			MarkComplete(OnlineSubsystem);
 		});
@@ -162,19 +162,19 @@ FTestDriver::~FTestDriver()
 
 bool FTestDriver::AddPipeline(FTestPipeline&& Pipeline, const FPipelineTestContext& TestContext)
 {
-	SubsystemType OnlineSubsystem = UE::Online::FOnlineServicesRegistry::Get().GetNamedServicesInstance(TestContext.InitOptions.ServiceType, NAME_None, NAME_None);
-	if (OnlineSubsystem == nullptr)
+	IOnlineServicesPtr OnlineServices = UE::Online::FOnlineServicesRegistry::Get().GetNamedServicesInstance(TestContext.InitOptions.ServiceType, NAME_None, NAME_None);
+	if (OnlineServices == nullptr)
 	{
 		return false;
 	}
 	else
 	{
-		SubsystemInstances.Emplace(OnlineSubsystem, MoveTemp(Pipeline));
+		SubsystemInstances.Emplace(OnlineServices, MoveTemp(Pipeline));
 		return true;
 	}
 }
 
-void FTestDriver::MarkComplete(SubsystemType Key)
+void FTestDriver::MarkComplete(IOnlineServicesPtr Key)
 {
 	CompletedInstances.Emplace(Key);
 }
@@ -183,7 +183,7 @@ void FTestDriver::RunToCompletion()
 {
 	LastTickTime = FPlatformTime::Seconds();
 
-	ForeachSubsystemInstance([](SubsystemType Subsystem, FTestPipeline& Func)
+	ForeachSubsystemInstance([](const IOnlineServicesPtr&  Subsystem, FTestPipeline& Func)
 		{
 			Func.Start();
 		});
@@ -192,7 +192,7 @@ void FTestDriver::RunToCompletion()
 	{
 		FPlatformProcess::Sleep(TICK_DURATION.GetTotalSeconds());
 		
-			ForeachSubsystemInstance([this](SubsystemType Subsystem, FTestPipeline& CurrentPipeline)
+			ForeachSubsystemInstance([this](const IOnlineServicesPtr& Services, FTestPipeline& CurrentPipeline)
 				{
 					const double BeforeTick = FPlatformTime::Seconds();
 					{ // Scope for INFO
@@ -207,10 +207,10 @@ void FTestDriver::RunToCompletion()
 					const double AfterTick = FPlatformTime::Seconds();
 					CurrentPipeline.EvaluatePlatformTickTime(AfterTick - BeforeTick);
 				});
-		ForeachSubsystemInstance([](SubsystemType Subsystem, FTestPipeline& Func)
+		ForeachSubsystemInstance([](const IOnlineServicesPtr& Services, FTestPipeline& Func)
 			{
 				//UE_LOG(LogOSSTests, Log, TEXT("INFO %s"), *Func.InfoString());
-				Func(Subsystem);
+				Func(Services);
 			});
 		FlushCompleted();
 	}
@@ -223,7 +223,7 @@ void FTestDriver::RunToCompletion()
 
 void FTestDriver::FlushCompleted()
 {
-	for (SubsystemType Key : CompletedInstances)
+	for (const IOnlineServicesPtr& Key : CompletedInstances)
 	{
 		SubsystemInstances.Find(Key)->OnPreRelease();
 		SubsystemInstances.Remove(Key);
diff --git a/Engine/Source/Programs/Online/OnlineTestsCore/Source/Private/Tests/Achievements/DisplayAchievementUITests.cpp b/Engine/Source/Programs/Online/OnlineTestsCore/Source/Private/Tests/Achievements/DisplayAchievementUITests.cpp
--- a/Engine/Source/Programs/Online/OnlineTestsCore/Source/Private/Tests/Achievements/DisplayAchievementUITests.cpp
+++ b/Engine/Source/Programs/Online/OnlineTestsCore/Source/Private/Tests/Achievements/DisplayAchievementUITests.cpp
@@ -14,7 +14,7 @@
 ACHIEVEMENTS_TEST_CASE("Display Achievement UI (Invalid User)", EG_ACHIEVEMENTS_DISPLAYACHIEVEMENTUI_TAG)
 {
 	GetPipeline()
-		.EmplaceLambda([](SubsystemType OnlineSubsystem)
+		.EmplaceLambda([](const IOnlineServicesPtr& OnlineSubsystem) 
 			{
 				UE::Online::FDisplayAchievementUI::Params OpParams;
 
@@ -39,7 +39,7 @@ ACHIEVEMENTS_TEST_CASE("Display Achievement UI (Invalid State)",  EG_ACHIEVEMENT
 	QueryDefinitionsHelperParams.OpParams.LocalAccountId = AccountId;
 
 	LoginPipeline.EmplaceStep<FQueryAchievementDefinitionsHelper>(MoveTemp(QueryDefinitionsHelperParams))
-		.EmplaceLambda([&AccountId](SubsystemType OnlineSubsystem)
+		.EmplaceLambda([&AccountId](const IOnlineServicesPtr& OnlineSubsystem) 
 		{
 			UE::Online::FDisplayAchievementUI::Params Params;
 			Params.LocalAccountId = AccountId;
@@ -67,7 +67,7 @@ ACHIEVEMENTS_TEST_CASE("Display Achievement UI (Not Found)", EG_ACHIEVEMENTS_DIS
 
 	LoginPipeline.EmplaceStep<FQueryAchievementDefinitionsHelper>(MoveTemp(QueryDefinitionsHelperParams))
 		.EmplaceStep<FQueryAchievementStatesHelper>(MoveTemp(QueryStatesHelperParams))
-		.EmplaceLambda([&AccountId](SubsystemType OnlineSubsystem)
+		.EmplaceLambda([&AccountId](const IOnlineServicesPtr& OnlineSubsystem) 
 		{
 			UE::Online::FDisplayAchievementUI::Params Params;
 			Params.LocalAccountId = AccountId;
@@ -96,7 +96,7 @@ ACHIEVEMENTS_TEST_CASE("Display Achievement UI (Success)", EG_ACHIEVEMENTS_DISPL
 
 	LoginPipeline.EmplaceStep<FQueryAchievementDefinitionsHelper>(MoveTemp(QueryDefinitionsHelperParams))
 		.EmplaceStep<FQueryAchievementStatesHelper>(MoveTemp(QueryStatesHelperParams))
-		.EmplaceLambda([&AccountId](SubsystemType OnlineSubsystem)
+		.EmplaceLambda([&AccountId](const IOnlineServicesPtr& OnlineSubsystem) 
 		{
 			UE::Online::FDisplayAchievementUI::Params Params;
 			Params.LocalAccountId = AccountId;
diff --git a/Engine/Source/Programs/Online/OnlineTestsCore/Source/Private/Tests/Achievements/GetAchievementDefinitionTests.cpp b/Engine/Source/Programs/Online/OnlineTestsCore/Source/Private/Tests/Achievements/GetAchievementDefinitionTests.cpp
--- a/Engine/Source/Programs/Online/OnlineTestsCore/Source/Private/Tests/Achievements/GetAchievementDefinitionTests.cpp
+++ b/Engine/Source/Programs/Online/OnlineTestsCore/Source/Private/Tests/Achievements/GetAchievementDefinitionTests.cpp
@@ -13,7 +13,7 @@
 ACHIEVEMENTS_TEST_CASE("Get Achievement Definition (Invalid User)", EG_ACHIEVEMENTS_GETACHIEVEMENTDEFINITION_TAG)
 {
 	GetPipeline()
-		.EmplaceLambda([](SubsystemType OnlineSubsystem)
+		.EmplaceLambda([](const IOnlineServicesPtr& OnlineSubsystem)
 			{
 				UE::Online::FGetAchievementDefinition::Params OpParams;
 
@@ -37,7 +37,7 @@ ACHIEVEMENTS_TEST_CASE("Get Achievement Definition (Invalid State)", EG_ACHIEVEM
 	UE::Online::FGetAchievementDefinition::Params GetAchievementDefinitionParams;
 	GetAchievementDefinitionParams.LocalAccountId = AccountId;
 
-	LoginPipeline.EmplaceLambda([&GetAchievementDefinitionParams](SubsystemType OnlineSubsystem)
+	LoginPipeline.EmplaceLambda([&GetAchievementDefinitionParams](const IOnlineServicesPtr& OnlineSubsystem)
 		{
 			UE::Online::IAchievementsPtr AchievementsInterface = OnlineSubsystem->GetAchievementsInterface();
 			const UE::Online::TOnlineResult<UE::Online::FGetAchievementDefinition> Result = AchievementsInterface->GetAchievementDefinition(MoveTemp(GetAchievementDefinitionParams));
@@ -62,7 +62,7 @@ ACHIEVEMENTS_TEST_CASE("Get Achievement Definition (Not Found)", EG_ACHIEVEMENTS
 	GetParams.AchievementId = TEXT("UnknownAchievement");
 
 	LoginPipeline.EmplaceStep<FQueryAchievementDefinitionsHelper>(MoveTemp(QueryDefinitionsHelperParams))
-		.EmplaceLambda([&GetParams](SubsystemType OnlineSubsystem)
+		.EmplaceLambda([&GetParams](const IOnlineServicesPtr& OnlineSubsystem)
 		{
 			UE::Online::IAchievementsPtr AchievementsInterface = OnlineSubsystem->GetAchievementsInterface();
 			const UE::Online::TOnlineResult<UE::Online::FGetAchievementDefinition> Result = AchievementsInterface->GetAchievementDefinition(MoveTemp(GetParams));
@@ -89,14 +89,14 @@ ACHIEVEMENTS_TEST_CASE("Get Achievement Definitions (Success)", EG_ACHIEVEMENTS_
 	GetParams.LocalAccountId = AccountId;
 
 	LoginPipeline.EmplaceStep<FQueryAchievementDefinitionsHelper>(MoveTemp(QueryDefinitionsHelperParams))
-		.EmplaceLambda([&GetParams, &AccountId, &AchievementIds](SubsystemType OnlineSubsystem)
+		.EmplaceLambda([&GetParams, &AccountId, &AchievementIds](const IOnlineServicesPtr& OnlineSubsystem)
 		{
 			UE::Online::IAchievementsPtr AchievementsInterface = OnlineSubsystem->GetAchievementsInterface();
 			UE::Online::TOnlineResult<UE::Online::FGetAchievementIds> Result = AchievementsInterface->GetAchievementIds(MoveTemp(GetParams));
 			REQUIRE_OP(Result);
 			AchievementIds = MoveTemp(Result.GetOkValue().AchievementIds);
 		})
-		.EmplaceLambda([&AccountId, &AchievementIds](SubsystemType OnlineSubsystem)
+		.EmplaceLambda([&AccountId, &AchievementIds](const IOnlineServicesPtr& OnlineSubsystem)
 		{
 			REQUIRE(!AchievementIds.IsEmpty());
 			for (const FString& AchievementId : AchievementIds)
diff --git a/Engine/Source/Programs/Online/OnlineTestsCore/Source/Private/Tests/Achievements/GetAchievementIdsTests.cpp b/Engine/Source/Programs/Online/OnlineTestsCore/Source/Private/Tests/Achievements/GetAchievementIdsTests.cpp
--- a/Engine/Source/Programs/Online/OnlineTestsCore/Source/Private/Tests/Achievements/GetAchievementIdsTests.cpp
+++ b/Engine/Source/Programs/Online/OnlineTestsCore/Source/Private/Tests/Achievements/GetAchievementIdsTests.cpp
@@ -13,7 +13,7 @@
 ACHIEVEMENTS_TEST_CASE("Get Achievement Ids (Invalid User)", EG_ACHIEVEMENTS_GETACHIEVEMENTIDS_TAG)
 {
 	GetPipeline()
-		.EmplaceLambda([](SubsystemType OnlineSubsystem)
+		.EmplaceLambda([](const IOnlineServicesPtr& OnlineSubsystem)
 			{
 				UE::Online::FGetAchievementIds::Params OpParams;
 
@@ -37,7 +37,7 @@ ACHIEVEMENTS_TEST_CASE("Get Achievement Ids (Invalid State)", EG_ACHIEVEMENTS_GE
 	UE::Online::FGetAchievementIds::Params OpParams;
 	OpParams.LocalAccountId = AccountId;
 
-	LoginPipeline.EmplaceLambda([&OpParams](SubsystemType OnlineSubsystem)
+	LoginPipeline.EmplaceLambda([&OpParams](const IOnlineServicesPtr& OnlineSubsystem)
 		{
 			IAchievementsPtr AchievementsInterface = OnlineSubsystem->GetAchievementsInterface();
 			TOnlineResult<UE::Online::FGetAchievementIds> Result = AchievementsInterface->GetAchievementIds(MoveTemp(OpParams));
@@ -61,7 +61,7 @@ ACHIEVEMENTS_TEST_CASE("Get Achievement Ids (Success)", EG_ACHIEVEMENTS_GETACHIE
 	GetParams.LocalAccountId = AccountId;
 
 	LoginPipeline.EmplaceStep<FQueryAchievementDefinitionsHelper>(MoveTemp(QueryDefinitionsHelperParams))
-		.EmplaceLambda([&GetParams](SubsystemType OnlineSubsystem)
+		.EmplaceLambda([&GetParams](const IOnlineServicesPtr& OnlineSubsystem)
 		{
 			UE::Online::IAchievementsPtr AchievementsInterface = OnlineSubsystem->GetAchievementsInterface();
 			const UE::Online::TOnlineResult<UE::Online::FGetAchievementIds> Result = AchievementsInterface->GetAchievementIds(MoveTemp(GetParams));
diff --git a/Engine/Source/Programs/Online/OnlineTestsCore/Source/Private/Tests/Achievements/GetAchievementStateTests.cpp b/Engine/Source/Programs/Online/OnlineTestsCore/Source/Private/Tests/Achievements/GetAchievementStateTests.cpp
--- a/Engine/Source/Programs/Online/OnlineTestsCore/Source/Private/Tests/Achievements/GetAchievementStateTests.cpp
+++ b/Engine/Source/Programs/Online/OnlineTestsCore/Source/Private/Tests/Achievements/GetAchievementStateTests.cpp
@@ -28,7 +28,7 @@ ACHIEVEMENTS_TEST_CASE("Get Achievement State (Invalid State)", EG_ACHIEVEMENTS_
 	GetParams.AchievementId = TEXT("test_unlockachievements");
 
 	LoginPipeline.EmplaceStep<FQueryAchievementDefinitionsHelper>(MoveTemp(QueryDefinitionsHelperParams))
-		.EmplaceLambda([&GetParams](SubsystemType OnlineSubsystem)
+		.EmplaceLambda([&GetParams](const IOnlineServicesPtr& OnlineSubsystem)
 		{
 			UE::Online::IAchievementsPtr AchievementsInterface = OnlineSubsystem->GetAchievementsInterface();
 			const UE::Online::TOnlineResult<UE::Online::FGetAchievementState> Result = AchievementsInterface->GetAchievementState(MoveTemp(GetParams));
@@ -57,7 +57,7 @@ ACHIEVEMENTS_TEST_CASE("Get Achievement State (Not Found)", EG_ACHIEVEMENTS_GETA
 
 	LoginPipeline.EmplaceStep<FQueryAchievementDefinitionsHelper>(MoveTemp(QueryDefinitionsHelperParams))
 		.EmplaceStep<FQueryAchievementStatesHelper>(MoveTemp(QueryStatesHelperParams))
-		.EmplaceLambda([&GetParams](SubsystemType OnlineSubsystem)
+		.EmplaceLambda([&GetParams](const IOnlineServicesPtr& OnlineSubsystem)
 	{
 		UE::Online::IAchievementsPtr AchievementsInterface = OnlineSubsystem->GetAchievementsInterface();
 		const UE::Online::TOnlineResult<UE::Online::FGetAchievementState> Result = AchievementsInterface->GetAchievementState(MoveTemp(GetParams));
@@ -86,7 +86,7 @@ ACHIEVEMENTS_TEST_CASE("Get Achievement State (Success)", EG_ACHIEVEMENTS_GETACH
 
 	LoginPipeline.EmplaceStep<FQueryAchievementDefinitionsHelper>(MoveTemp(QueryDefinitionsHelperParams))
 	.EmplaceStep<FQueryAchievementStatesHelper>(MoveTemp(QueryStatesHelperParams))
-	.EmplaceLambda([&GetParams](SubsystemType OnlineSubsystem)
+	.EmplaceLambda([&GetParams](const IOnlineServicesPtr& OnlineSubsystem)
 	{
 		UE::Online::IAchievementsPtr AchievementsInterface = OnlineSubsystem->GetAchievementsInterface();
 		const UE::Online::TOnlineResult<UE::Online::FGetAchievementState> Result = AchievementsInterface->GetAchievementState(MoveTemp(GetParams));
diff --git a/Engine/Source/Programs/Online/OnlineTestsCore/Source/Private/Tests/Achievements/QueryAchievementStatesTests.cpp b/Engine/Source/Programs/Online/OnlineTestsCore/Source/Private/Tests/Achievements/QueryAchievementStatesTests.cpp
--- a/Engine/Source/Programs/Online/OnlineTestsCore/Source/Private/Tests/Achievements/QueryAchievementStatesTests.cpp
+++ b/Engine/Source/Programs/Online/OnlineTestsCore/Source/Private/Tests/Achievements/QueryAchievementStatesTests.cpp
@@ -65,7 +65,7 @@ ACHIEVEMENTS_TEST_CASE("Get Achievement State (Invalid User)", EG_ACHIEVEMENTS_Q
 	FAccountId AccountId;
 
 	GetLoginPipeline({ AccountId })
-		.EmplaceLambda([](SubsystemType OnlineSubsystem)
+		.EmplaceLambda([](const IOnlineServicesPtr& OnlineSubsystem)
 			{
 				UE::Online::FGetAchievementState::Params Params;
 				Params.AchievementId = TEXT("test_unlockachievements");
diff --git a/Engine/Source/Programs/Online/OnlineTestsCore/Source/Private/Tests/Achievements/UnlockAchievementsTests.cpp b/Engine/Source/Programs/Online/OnlineTestsCore/Source/Private/Tests/Achievements/UnlockAchievementsTests.cpp
--- a/Engine/Source/Programs/Online/OnlineTestsCore/Source/Private/Tests/Achievements/UnlockAchievementsTests.cpp
+++ b/Engine/Source/Programs/Online/OnlineTestsCore/Source/Private/Tests/Achievements/UnlockAchievementsTests.cpp
@@ -179,18 +179,18 @@ ACHIEVEMENTS_TEST_CASE("Unlock Title-Managed Achievements By Stats", EG_ACHIEVEM
 
 	LoginPipeline
 		.EmplaceStep<FQueryAchievementDefinitionsHelper>(MoveTemp(QueryDefinitionsHelperParams))
-		.EmplaceAsyncLambda([&AccountId, &AchievementId, &AchievementProgress](FAsyncLambdaResult Promise, SubsystemType Services)
+		.EmplaceAsyncLambda([&AccountId, &AchievementId, &AchievementProgress](FAsyncLambdaResult Promise, const IOnlineServicesPtr& Services) 
 			{ 
 				CheckAchievementState_Fixture(Promise, Services, AccountId, AchievementId, AchievementProgress);
 			})
-		.EmplaceAsyncLambda([&AccountId, &StatNameKey, &StatValue](FAsyncLambdaResult Promise, SubsystemType Services)
+		.EmplaceAsyncLambda([&AccountId, &StatNameKey, &StatValue](FAsyncLambdaResult Promise, const IOnlineServicesPtr& Services) 
 			{
 				FUserStats UsersStats{ AccountId, {{ StatNameKey, StatValue }} };
 				TArray<FUserStats> UpdateUsersStats{ UsersStats };
 
 				UpdateStats_Fixture(Promise, Services, AccountId, UpdateUsersStats);
 			})
-		.EmplaceAsyncLambda([&AccountId, &AchievementId, &UpdatedAchievementProgress](FAsyncLambdaResult Promise, SubsystemType Services)
+		.EmplaceAsyncLambda([&AccountId, &AchievementId, &UpdatedAchievementProgress](const FAsyncLambdaResult& Promise, const IOnlineServicesPtr& Services) 
 			{ 
 				CheckAchievementState_Fixture(Promise, Services, AccountId, AchievementId, UpdatedAchievementProgress);
 			});
@@ -224,27 +224,28 @@ ACHIEVEMENTS_TEST_CASE("Unlock Title-Managed Achievements By Stats With Multiple
 
 	LoginPipeline
 		.EmplaceStep<FQueryAchievementDefinitionsHelper>(MoveTemp(QueryDefinitionsHelperParams))
-		.EmplaceAsyncLambda([&AccountId, &AchievementId, &AchievementProgress](FAsyncLambdaResult Promise, SubsystemType Services)
+		.EmplaceAsyncLambda([&AccountId, &AchievementId, &AchievementProgress](const FAsyncLambdaResult& Promise, const IOnlineServicesPtr& Services) 
 			{ 
 				CheckAchievementState_Fixture(Promise, Services, AccountId, AchievementId, AchievementProgress);
 			})
-		.EmplaceAsyncLambda([&AccountId, &FirstStatNameKey, &FirstStatValue, &SecondStatNameKey, &SecondStatValue](FAsyncLambdaResult Promise, SubsystemType Services)
+		.EmplaceAsyncLambda([&AccountId, &FirstStatNameKey, &FirstStatValue, &SecondStatNameKey, &SecondStatValue](const FAsyncLambdaResult& Promise, const IOnlineServicesPtr& Services) 
 			{
 				FUserStats UserStats{ AccountId, {{ FirstStatNameKey, FirstStatValue}, {SecondStatNameKey, SecondStatValue} } };
 				TArray<FUserStats> UpdateUsersStats{ UserStats }; //Partially meet conditions
 				UpdateStats_Fixture(Promise, Services, AccountId, UpdateUsersStats);
 			})
-		.EmplaceAsyncLambda([&AccountId, &AchievementId, &AchievementProgress](FAsyncLambdaResult Promise, SubsystemType Services)
+		.EmplaceAsyncLambda([&AccountId, &AchievementId, &AchievementProgress](const FAsyncLambdaResult& Promise, const IOnlineServicesPtr& Services) 
 			{
 				CheckAchievementState_Fixture(Promise, Services, AccountId, AchievementId, AchievementProgress);
 			})
-		.EmplaceAsyncLambda([&AccountId, &FirstStatNameKey, &ThirdStatValue, &SecondStatNameKey, &SecondStatValue](FAsyncLambdaResult Promise, SubsystemType Services)
+		.EmplaceAsyncLambda([&AccountId, &FirstStatNameKey, &ThirdStatValue, &SecondStatNameKey, &SecondStatValue]
+			(const FAsyncLambdaResult& Promise, const IOnlineServicesPtr& Services) 
 			{
 				FUserStats UpdatedUsersStats{ AccountId, {{ FirstStatNameKey, ThirdStatValue}, {SecondStatNameKey, SecondStatValue } } };
 				TArray<FUserStats> UpdateUsersStats{ UpdatedUsersStats}; //Meet all conditions
 				UpdateStats_Fixture(Promise, Services, AccountId, UpdateUsersStats);
 			})
-		.EmplaceAsyncLambda([&AccountId, &AchievementId, &UpdatedAchievementProgress](FAsyncLambdaResult Promise, SubsystemType Services)
+		.EmplaceAsyncLambda([&AccountId, &AchievementId, &UpdatedAchievementProgress](const FAsyncLambdaResult& Promise, const IOnlineServicesPtr& Services) 
 			{ 
 				CheckAchievementState_Fixture(Promise, Services, AccountId, AchievementId, UpdatedAchievementProgress);
 			});
diff --git a/Engine/Source/Programs/Online/OnlineTestsCore/Source/Private/Tests/Commerce/CheckoutTests.cpp b/Engine/Source/Programs/Online/OnlineTestsCore/Source/Private/Tests/Commerce/CheckoutTests.cpp
--- a/Engine/Source/Programs/Online/OnlineTestsCore/Source/Private/Tests/Commerce/CheckoutTests.cpp
+++ b/Engine/Source/Programs/Online/OnlineTestsCore/Source/Private/Tests/Commerce/CheckoutTests.cpp
@@ -1,9 +1,11 @@
 // Copyright Epic Games, Inc. All Rights Reserved.
 
-#include "Helpers/Commerce/QueryOffersHelper.h"
-#include "Helpers/Commerce/GetOffersByIdHelper.h"
-#include "Helpers/Commerce/CheckoutHelper.h"
+#include "CoreGlobals.h"
 #include "Helpers/Auth/AuthLogout.h"
+#include "Helpers/Commerce/CheckoutHelper.h"
+#include "Helpers/Commerce/GetOffersByIdHelper.h"
+#include "Helpers/Commerce/QueryOffersHelper.h"
+#include "Misc/ConfigCacheIni.h"
 
 #define COMMERCE_TAG "[suite_commerce]"
 #define COMMERCE_CHECKOUT_TAG COMMERCE_TAG "[checkout]"
@@ -26,14 +28,14 @@ COMMERCE_TEST_CASE("Verify that Checkout returns a fail message if the local use
 	bool bLogout = false;
 
 	LoginPipeline
-		.EmplaceLambda([&AccountId, AccountPlatformUserId](SubsystemType OnlineSubsystem)
+		.EmplaceLambda([&AccountId, AccountPlatformUserId](const IOnlineServicesPtr& OnlineSubsystem)
 			{
 				UE::Online::IAuthPtr OnlineAuthPtr = OnlineSubsystem->GetAuthInterface();
 				REQUIRE(OnlineAuthPtr);
-				UE::Online::TOnlineResult<UE::Online::FAuthGetLocalOnlineUserByOnlineAccountId> UserPlatfromUserIdResult = OnlineAuthPtr->GetLocalOnlineUserByOnlineAccountId({ AccountId });
-				REQUIRE(UserPlatfromUserIdResult.IsOk());
-				CHECK(UserPlatfromUserIdResult.TryGetOkValue() != nullptr);
-				*AccountPlatformUserId = UserPlatfromUserIdResult.TryGetOkValue()->AccountInfo->PlatformUserId;
+				UE::Online::TOnlineResult<UE::Online::FAuthGetLocalOnlineUserByOnlineAccountId> UserPlatformUserIdResult = OnlineAuthPtr->GetLocalOnlineUserByOnlineAccountId({ AccountId });
+				REQUIRE(UserPlatformUserIdResult.IsOk());
+				CHECK(UserPlatformUserIdResult.TryGetOkValue() != nullptr);
+				*AccountPlatformUserId = UserPlatformUserIdResult.TryGetOkValue()->AccountInfo->PlatformUserId;
 			})
 		.EmplaceStep<FAuthLogoutStep>(MoveTemp(AccountPlatformUserId))
 		.EmplaceStep<FCheckoutHelper>(MoveTemp(CheckoutHelperParams));
diff --git a/Engine/Source/Programs/Online/OnlineTestsCore/Source/Private/Tests/Commerce/GetOffersByIdTests.cpp b/Engine/Source/Programs/Online/OnlineTestsCore/Source/Private/Tests/Commerce/GetOffersByIdTests.cpp
--- a/Engine/Source/Programs/Online/OnlineTestsCore/Source/Private/Tests/Commerce/GetOffersByIdTests.cpp
+++ b/Engine/Source/Programs/Online/OnlineTestsCore/Source/Private/Tests/Commerce/GetOffersByIdTests.cpp
@@ -1,8 +1,10 @@
 // Copyright Epic Games, Inc. All Rights Reserved.
 
-#include "Helpers/Commerce/QueryOffersHelper.h"
-#include "Helpers/Commerce/GetOffersByIdHelper.h"
+#include "CoreGlobals.h"
 #include "Helpers/Auth/AuthLogout.h"
+#include "Helpers/Commerce/GetOffersByIdHelper.h"
+#include "Helpers/Commerce/QueryOffersHelper.h"
+#include "Misc/ConfigCacheIni.h"
 
 #define COMMERCE_TAG "[suite_commerce]"
 #define COMMERCE_GETOFFERSBYID_TAG COMMERCE_TAG "[getoffersbyid]"
@@ -25,14 +27,14 @@ COMMERCE_TEST_CASE("Verify that GetOffersById returns a fail message if the loca
 	bool bLogout = false;
 
 	LoginPipeline
-		.EmplaceLambda([&AccountId, AccountPlatformUserId](SubsystemType OnlineSubsystem)
+		.EmplaceLambda([&AccountId, AccountPlatformUserId](const IOnlineServicesPtr& OnlineSubsystem)
 			{
 				UE::Online::IAuthPtr OnlineAuthPtr = OnlineSubsystem->GetAuthInterface();
 				REQUIRE(OnlineAuthPtr);
-				UE::Online::TOnlineResult<UE::Online::FAuthGetLocalOnlineUserByOnlineAccountId> UserPlatfromUserIdResult = OnlineAuthPtr->GetLocalOnlineUserByOnlineAccountId({ AccountId });
-				REQUIRE(UserPlatfromUserIdResult.IsOk());
-				CHECK(UserPlatfromUserIdResult.TryGetOkValue() != nullptr);
-				*AccountPlatformUserId = UserPlatfromUserIdResult.TryGetOkValue()->AccountInfo->PlatformUserId;
+				UE::Online::TOnlineResult<UE::Online::FAuthGetLocalOnlineUserByOnlineAccountId> UserPlatformUserIdResult = OnlineAuthPtr->GetLocalOnlineUserByOnlineAccountId({ AccountId });
+				REQUIRE(UserPlatformUserIdResult.IsOk());
+				CHECK(UserPlatformUserIdResult.TryGetOkValue() != nullptr);
+				*AccountPlatformUserId = UserPlatformUserIdResult.TryGetOkValue()->AccountInfo->PlatformUserId;
 			})
 		.EmplaceStep<FAuthLogoutStep>(MoveTemp(AccountPlatformUserId))
 		.EmplaceStep<FGetOffersByIdHelper>(MoveTemp(GetOffersByIdHelperParams));
diff --git a/Engine/Source/Programs/Online/OnlineTestsCore/Source/Private/Tests/Commerce/GetOffersTests.cpp b/Engine/Source/Programs/Online/OnlineTestsCore/Source/Private/Tests/Commerce/GetOffersTests.cpp
--- a/Engine/Source/Programs/Online/OnlineTestsCore/Source/Private/Tests/Commerce/GetOffersTests.cpp
+++ b/Engine/Source/Programs/Online/OnlineTestsCore/Source/Private/Tests/Commerce/GetOffersTests.cpp
@@ -1,8 +1,10 @@
 // Copyright Epic Games, Inc. All Rights Reserved.
 
-#include "Helpers/Commerce/QueryOffersHelper.h"
-#include "Helpers/Commerce/GetOffersHelper.h"
+#include "CoreGlobals.h"
 #include "Helpers/Auth/AuthLogout.h"
+#include "Helpers/Commerce/GetOffersHelper.h"
+#include "Helpers/Commerce/QueryOffersHelper.h"
+#include "Misc/ConfigCacheIni.h"
 
 #define COMMERCE_TAG "[suite_commerce]"
 #define COMMERCE_GETOFFERS_TAG COMMERCE_TAG "[getoffers]"
@@ -25,14 +27,14 @@ COMMERCE_TEST_CASE("Verify that GetOffers returns a fail message if the local us
 	bool bLogout = false;
 
 	LoginPipeline
-		.EmplaceLambda([&AccountId, AccountPlatformUserId](SubsystemType OnlineSubsystem)
+		.EmplaceLambda([&AccountId, AccountPlatformUserId](const IOnlineServicesPtr& OnlineSubsystem)
 			{
 				UE::Online::IAuthPtr OnlineAuthPtr = OnlineSubsystem->GetAuthInterface();
 				REQUIRE(OnlineAuthPtr);
-				UE::Online::TOnlineResult<UE::Online::FAuthGetLocalOnlineUserByOnlineAccountId> UserPlatfromUserIdResult = OnlineAuthPtr->GetLocalOnlineUserByOnlineAccountId({ AccountId });
-				REQUIRE(UserPlatfromUserIdResult.IsOk());
-				CHECK(UserPlatfromUserIdResult.TryGetOkValue() != nullptr);
-				*AccountPlatformUserId = UserPlatfromUserIdResult.TryGetOkValue()->AccountInfo->PlatformUserId;
+				UE::Online::TOnlineResult<UE::Online::FAuthGetLocalOnlineUserByOnlineAccountId> UserPlatformUserIdResult = OnlineAuthPtr->GetLocalOnlineUserByOnlineAccountId({ AccountId });
+				REQUIRE(UserPlatformUserIdResult.IsOk());
+				CHECK(UserPlatformUserIdResult.TryGetOkValue() != nullptr);
+				*AccountPlatformUserId = UserPlatformUserIdResult.TryGetOkValue()->AccountInfo->PlatformUserId;
 			})
 		.EmplaceStep<FAuthLogoutStep>(MoveTemp(AccountPlatformUserId))
 		.EmplaceStep<FGetOffersHelper>(MoveTemp(GetOffersHelperParams));
diff --git a/Engine/Source/Programs/Online/OnlineTestsCore/Source/Private/Tests/Commerce/QueryOffersByIdTests.cpp b/Engine/Source/Programs/Online/OnlineTestsCore/Source/Private/Tests/Commerce/QueryOffersByIdTests.cpp
--- a/Engine/Source/Programs/Online/OnlineTestsCore/Source/Private/Tests/Commerce/QueryOffersByIdTests.cpp
+++ b/Engine/Source/Programs/Online/OnlineTestsCore/Source/Private/Tests/Commerce/QueryOffersByIdTests.cpp
@@ -25,14 +25,14 @@ COMMERCE_TEST_CASE("Verify that QueryOffersById returns a fail message if the lo
 	bool bLogout = false;
 
 	LoginPipeline
-		.EmplaceLambda([&AccountId, AccountPlatformUserId](SubsystemType OnlineSubsystem)
+		.EmplaceLambda([&AccountId, AccountPlatformUserId](const IOnlineServicesPtr& OnlineSubsystem)
 			{
 				UE::Online::IAuthPtr OnlineAuthPtr = OnlineSubsystem->GetAuthInterface();
 				REQUIRE(OnlineAuthPtr);
-				UE::Online::TOnlineResult<UE::Online::FAuthGetLocalOnlineUserByOnlineAccountId> UserPlatfromUserIdResult = OnlineAuthPtr->GetLocalOnlineUserByOnlineAccountId({ AccountId });
-				REQUIRE(UserPlatfromUserIdResult.IsOk());
-				CHECK(UserPlatfromUserIdResult.TryGetOkValue() != nullptr);
-				*AccountPlatformUserId = UserPlatfromUserIdResult.TryGetOkValue()->AccountInfo->PlatformUserId;
+				UE::Online::TOnlineResult<UE::Online::FAuthGetLocalOnlineUserByOnlineAccountId> UserPlatformUserIdResult = OnlineAuthPtr->GetLocalOnlineUserByOnlineAccountId({ AccountId });
+				REQUIRE(UserPlatformUserIdResult.IsOk());
+				CHECK(UserPlatformUserIdResult.TryGetOkValue() != nullptr);
+				*AccountPlatformUserId = UserPlatformUserIdResult.TryGetOkValue()->AccountInfo->PlatformUserId;
 			})
 		.EmplaceStep<FAuthLogoutStep>(MoveTemp(AccountPlatformUserId))
 		.EmplaceStep<FQueryOffersByIdHelper>(MoveTemp(QueryOffersByIdHelperParams));
diff --git a/Engine/Source/Programs/Online/OnlineTestsCore/Source/Private/Tests/Commerce/QueryOffersTests.cpp b/Engine/Source/Programs/Online/OnlineTestsCore/Source/Private/Tests/Commerce/QueryOffersTests.cpp
--- a/Engine/Source/Programs/Online/OnlineTestsCore/Source/Private/Tests/Commerce/QueryOffersTests.cpp
+++ b/Engine/Source/Programs/Online/OnlineTestsCore/Source/Private/Tests/Commerce/QueryOffersTests.cpp
@@ -1,8 +1,10 @@
 // Copyright Epic Games, Inc. All Rights Reserved.
 
+#include "CoreGlobals.h"
+#include "Helpers/Auth/AuthLogout.h"
 #include "Helpers/Commerce/GetOffersHelper.h"
 #include "Helpers/Commerce/QueryOffersHelper.h"
-#include "Helpers/Auth/AuthLogout.h"
+#include "Misc/ConfigCacheIni.h"
 
 #define COMMERCE_TAG "[suite_commerce]"
 #define COMMERCE_QUERYOFFERS_TAG COMMERCE_TAG "[queryoffers]"
@@ -31,14 +33,14 @@ COMMERCE_TEST_CASE("Verify that QueryOffers returns a fail message if the local
 	bool bLogout = false;
 
 	LoginPipeline
-		.EmplaceLambda([&AccountId, AccountPlatformUserId](SubsystemType OnlineSubsystem)
+		.EmplaceLambda([&AccountId, AccountPlatformUserId](const IOnlineServicesPtr& OnlineServices)
 			{
-				UE::Online::IAuthPtr OnlineAuthPtr = OnlineSubsystem->GetAuthInterface();
+				UE::Online::IAuthPtr OnlineAuthPtr = OnlineServices->GetAuthInterface();
 				REQUIRE(OnlineAuthPtr);
-				UE::Online::TOnlineResult<UE::Online::FAuthGetLocalOnlineUserByOnlineAccountId> UserPlatfromUserIdResult = OnlineAuthPtr->GetLocalOnlineUserByOnlineAccountId({ AccountId });
-				REQUIRE(UserPlatfromUserIdResult.IsOk());
-				CHECK(UserPlatfromUserIdResult.TryGetOkValue() != nullptr);
-				*AccountPlatformUserId = UserPlatfromUserIdResult.TryGetOkValue()->AccountInfo->PlatformUserId;
+				UE::Online::TOnlineResult<UE::Online::FAuthGetLocalOnlineUserByOnlineAccountId> UserPlatformUserIdResult = OnlineAuthPtr->GetLocalOnlineUserByOnlineAccountId({ AccountId });
+				REQUIRE(UserPlatformUserIdResult.IsOk());
+				CHECK(UserPlatformUserIdResult.TryGetOkValue() != nullptr);
+				*AccountPlatformUserId = UserPlatformUserIdResult.TryGetOkValue()->AccountInfo->PlatformUserId;
 			})
 		.EmplaceStep<FAuthLogoutStep>(MoveTemp(AccountPlatformUserId))
 		.EmplaceStep<FQueryOffersHelper>(MoveTemp(QueryOffersHelperParams));
diff --git a/Engine/Source/Programs/Online/OnlineTestsCore/Source/Private/Tests/Leaderboards/LeaderboardsFixture.cpp b/Engine/Source/Programs/Online/OnlineTestsCore/Source/Private/Tests/Leaderboards/LeaderboardsFixture.cpp
--- a/Engine/Source/Programs/Online/OnlineTestsCore/Source/Private/Tests/Leaderboards/LeaderboardsFixture.cpp
+++ b/Engine/Source/Programs/Online/OnlineTestsCore/Source/Private/Tests/Leaderboards/LeaderboardsFixture.cpp
@@ -10,7 +10,7 @@ const UE::Online::FLeaderboardEntry* LeaderboardEntryOf(const TArray<UE::Online:
 
 FTestPipeline&& LeaderboardsFixture_NUsers(FTestPipeline&& Pipeline, TArray<FAccountId*>& AccountIds, const uint32& NumUsersToLogin)
 {
-	return Pipeline.EmplaceAsyncLambda([&](FAsyncLambdaResult Promise, SubsystemType Services) {
+	return Pipeline.EmplaceAsyncLambda([&](FAsyncLambdaResult Promise, const IOnlineServicesPtr& Services) {
 		TArray<FUserStats> UpdateUsersStats;
 		UpdateUsersStats.Reserve(NumUsersToLogin);
 		uint32 Counter = 0;
diff --git a/Engine/Source/Programs/Online/OnlineTestsCore/Source/Private/Tests/Leaderboards/LeaderboardsReadEntriesAroundRankTests.cpp b/Engine/Source/Programs/Online/OnlineTestsCore/Source/Private/Tests/Leaderboards/LeaderboardsReadEntriesAroundRankTests.cpp
--- a/Engine/Source/Programs/Online/OnlineTestsCore/Source/Private/Tests/Leaderboards/LeaderboardsReadEntriesAroundRankTests.cpp
+++ b/Engine/Source/Programs/Online/OnlineTestsCore/Source/Private/Tests/Leaderboards/LeaderboardsReadEntriesAroundRankTests.cpp
@@ -17,10 +17,10 @@ LEADERBOARDS_TEST_CASE("Verify ReadEntriesAroundRank succeeds and returns entrie
 	uint32 NumUsersToLogin = 2;
 
 	LeaderboardsFixture_NUsers(MoveTemp(GetLoginPipeline({ AccountIdA, AccountIdB })), AccountIds, NumUsersToLogin)
-		.EmplaceAsyncLambda([&](FAsyncLambdaResult Promise, SubsystemType Services) { ReadEntriesAroundRank_2UsersFixture(Promise, Services, "Stat_Use_Set", AccountIdA, 0, 2, LeaderboardEntries); })
-		.EmplaceLambda([&](SubsystemType Services) { CheckEntries_2Users_Stat_Use_Set(LeaderboardEntries, AccountIdA, AccountIdB); })
-		.EmplaceAsyncLambda([&](FAsyncLambdaResult Promise, SubsystemType Services) { ReadEntriesAroundRank_2UsersFixture(Promise, Services, "Stat_Use_Smallest", AccountIdA, 0, 2, LeaderboardEntries); })
-		.EmplaceLambda([&](SubsystemType Services) { CheckEntries_2Users_Stat_Use_Smallest(LeaderboardEntries, AccountIdA, AccountIdB); });
+		.EmplaceAsyncLambda([&](const FAsyncLambdaResult& Promise, const IOnlineServicesPtr& Services) { ReadEntriesAroundRank_2UsersFixture(Promise, Services, "Stat_Use_Set", AccountIdA, 0, 2, LeaderboardEntries); })
+		.EmplaceLambda([&](const IOnlineServicesPtr& Services) { CheckEntries_2Users_Stat_Use_Set(LeaderboardEntries, AccountIdA, AccountIdB); })
+		.EmplaceAsyncLambda([&](const FAsyncLambdaResult& Promise, const IOnlineServicesPtr& Services) { ReadEntriesAroundRank_2UsersFixture(Promise, Services, "Stat_Use_Smallest", AccountIdA, 0, 2, LeaderboardEntries); })
+		.EmplaceLambda([&](const IOnlineServicesPtr& Services) { CheckEntries_2Users_Stat_Use_Smallest(LeaderboardEntries, AccountIdA, AccountIdB); });
 
 	RunToCompletion();
 }
@@ -33,7 +33,7 @@ LEADERBOARDS_TEST_CASE("Verify ReadEntriesAroundRank returns a fail message if t
 	TArray< FAccountId*> AccountIds = { &AccountId, &LocalUserInvalid };
 
 	LeaderboardsFixture_NUsers(MoveTemp(GetLoginPipeline({ AccountId })), AccountIds, 1)
-		.EmplaceAsyncLambda([&](FAsyncLambdaResult Promise, SubsystemType Services) { ReadEntriesAroundRankFixture(Promise, Services, "Stat_Use_Set", LocalUserInvalid, 0, 2, UE::Online::Errors::InvalidUser(), LeaderboardEntries); });
+		.EmplaceAsyncLambda([&](const FAsyncLambdaResult& Promise, const IOnlineServicesPtr& Services) { ReadEntriesAroundRankFixture(Promise, Services, "Stat_Use_Set", LocalUserInvalid, 0, 2, UE::Online::Errors::InvalidUser(), LeaderboardEntries); });
 
 	RunToCompletion();
 }
@@ -46,7 +46,7 @@ LEADERBOARDS_TEST_CASE("Verify ReadEntriesAroundRank returns a fail message if g
 	TArray< FAccountId*> AccountIds = { &AccountId };
 
 	LeaderboardsFixture_NUsers(MoveTemp(GetLoginPipeline({ AccountId })), AccountIds, 1)
-		.EmplaceAsyncLambda([&](FAsyncLambdaResult Promise, SubsystemType Services) { ReadEntriesAroundRankFixture(Promise, Services, "Stat_Use_Set", AccountId, 1, 2, UE::Online::Errors::InvalidParams(), LeaderboardEntries); });
+		.EmplaceAsyncLambda([&](const FAsyncLambdaResult& Promise, const IOnlineServicesPtr& Services) { ReadEntriesAroundRankFixture(Promise, Services, "Stat_Use_Set", AccountId, 1, 2, UE::Online::Errors::InvalidParams(), LeaderboardEntries); });
 
 	RunToCompletion();
 }
@@ -61,8 +61,8 @@ LEADERBOARDS_TEST_CASE("Verify ReadEntriesAroundRank returns all valid entries i
 	uint32 limit = 3;
 
 	LeaderboardsFixture_NUsers(MoveTemp(GetLoginPipeline({ AccountIdA, AccountIdB })), AccountIds, 2)
-		.EmplaceAsyncLambda([&](FAsyncLambdaResult Promise, SubsystemType Services) { ReadEntriesAroundRank_2UsersFixture(Promise, Services, "Stat_Use_Set", AccountIdA, rank, limit, LeaderboardEntries); })
-		.EmplaceLambda([&](SubsystemType Services) { CheckEntries_2Users_Stat_Use_Set(LeaderboardEntries, AccountIdA, AccountIdB); });
+		.EmplaceAsyncLambda([&](const FAsyncLambdaResult& Promise, const IOnlineServicesPtr& Services) { ReadEntriesAroundRank_2UsersFixture(Promise, Services, "Stat_Use_Set", AccountIdA, rank, limit, LeaderboardEntries); })
+		.EmplaceLambda([&](const IOnlineServicesPtr& Services) { CheckEntries_2Users_Stat_Use_Set(LeaderboardEntries, AccountIdA, AccountIdB); });
 
 	RunToCompletion();
 }
@@ -75,7 +75,7 @@ LEADERBOARDS_TEST_CASE("Verify ReadEntriesAroundRank returns fail message if giv
 	TArray< FAccountId*> AccountIds = { &AccountId };
 
 	LeaderboardsFixture_NUsers(MoveTemp(GetLoginPipeline({ AccountId })), AccountIds, 1)
-		.EmplaceAsyncLambda([&](FAsyncLambdaResult Promise, SubsystemType Services) { ReadEntriesAroundRankFixture(Promise, Services, "Stat_Use_Set", AccountId, 3, 0, UE::Online::Errors::InvalidParams(), LeaderboardEntries); });
+		.EmplaceAsyncLambda([&](const FAsyncLambdaResult& Promise, const IOnlineServicesPtr& Services) { ReadEntriesAroundRankFixture(Promise, Services, "Stat_Use_Set", AccountId, 3, 0, UE::Online::Errors::InvalidParams(), LeaderboardEntries); });
 
 	RunToCompletion();
 }
diff --git a/Engine/Source/Programs/Online/OnlineTestsCore/Source/Private/Tests/Leaderboards/LeaderboardsReadEntriesAroundUserTests.cpp b/Engine/Source/Programs/Online/OnlineTestsCore/Source/Private/Tests/Leaderboards/LeaderboardsReadEntriesAroundUserTests.cpp
--- a/Engine/Source/Programs/Online/OnlineTestsCore/Source/Private/Tests/Leaderboards/LeaderboardsReadEntriesAroundUserTests.cpp
+++ b/Engine/Source/Programs/Online/OnlineTestsCore/Source/Private/Tests/Leaderboards/LeaderboardsReadEntriesAroundUserTests.cpp
@@ -20,10 +20,10 @@ LEADERBOARDS_TEST_CASE("Verify that ReadEntriesAroundUser succeeds and returns e
 	uint32 Limit = 2;
 
 	LeaderboardsFixture_NUsers(MoveTemp(GetLoginPipeline({ AccountIdA, AccountIdB })), AccountIds, 2)
-		.EmplaceAsyncLambda([&](FAsyncLambdaResult Promise, SubsystemType Services) { ReadEntriesAroundUser_2UsersFixture(Promise, Services, "Stat_Use_Set", AccountIdB, AccountIdA, Offset, Limit, LeaderboardEntries); })
-		.EmplaceLambda([&](SubsystemType Services) { CheckEntries_2Users_Stat_Use_Set(LeaderboardEntries, AccountIdA, AccountIdB); })
-		.EmplaceAsyncLambda([&](FAsyncLambdaResult Promise, SubsystemType Services) { ReadEntriesAroundUser_2UsersFixture(Promise, Services, "Stat_Use_Smallest", AccountIdA, AccountIdB, -1, 2, LeaderboardEntries); })
-		.EmplaceLambda([&](SubsystemType Services) { CheckEntries_2Users_Stat_Use_Smallest(LeaderboardEntries, AccountIdA, AccountIdB); });
+		.EmplaceAsyncLambda([&](const FAsyncLambdaResult& Promise, const IOnlineServicesPtr& Services) { ReadEntriesAroundUser_2UsersFixture(Promise, Services, "Stat_Use_Set", AccountIdB, AccountIdA, Offset, Limit, LeaderboardEntries); })
+		.EmplaceLambda([&](const IOnlineServicesPtr& Services) { CheckEntries_2Users_Stat_Use_Set(LeaderboardEntries, AccountIdA, AccountIdB); })
+		.EmplaceAsyncLambda([&](const FAsyncLambdaResult& Promise, const IOnlineServicesPtr& Services) { ReadEntriesAroundUser_2UsersFixture(Promise, Services, "Stat_Use_Smallest", AccountIdA, AccountIdB, -1, 2, LeaderboardEntries); })
+		.EmplaceLambda([&](const IOnlineServicesPtr& Services) { CheckEntries_2Users_Stat_Use_Smallest(LeaderboardEntries, AccountIdA, AccountIdB); });
 
 	RunToCompletion();
 }
@@ -38,7 +38,7 @@ LEADERBOARDS_TEST_CASE("Verify ReadEntriesAroundUser returns a fail message if t
 	uint32 limit = 2;
 
 	LeaderboardsFixture_NUsers(MoveTemp(GetLoginPipeline({ AccountId })), AccountIds, 1)
-		.EmplaceAsyncLambda([&](FAsyncLambdaResult Promise, SubsystemType Services) {ReadEntriesAroundUserFixture(Promise, Services, "Stat_Use_Set", LocalUserInvalid, AccountId, offset, limit, UE::Online::Errors::InvalidUser(), LeaderboardEntries); });
+		.EmplaceAsyncLambda([&](const FAsyncLambdaResult& Promise, const IOnlineServicesPtr& Services) { ReadEntriesAroundUserFixture(Promise, Services, "Stat_Use_Set", LocalUserInvalid, AccountId, offset, limit, UE::Online::Errors::InvalidUser(), LeaderboardEntries); });
 
 	RunToCompletion();
 }
@@ -53,7 +53,7 @@ LEADERBOARDS_TEST_CASE("Verify ReadEntriesAroundUser returns a fail message if t
 	uint32 limit = 2;
 
 	LeaderboardsFixture_NUsers(MoveTemp(GetLoginPipeline({ LocalUser })), AccountIds, 1)
-		.EmplaceAsyncLambda([&](FAsyncLambdaResult Promise, SubsystemType Services) {ReadEntriesAroundUserFixture(Promise, Services, "Stat_Use_Set", LocalUser, AccountIdInvalid, offset, limit, UE::Online::Errors::InvalidUser(), LeaderboardEntries); });
+		.EmplaceAsyncLambda([&](const FAsyncLambdaResult& Promise, const IOnlineServicesPtr& Services) { ReadEntriesAroundUserFixture(Promise, Services, "Stat_Use_Set", LocalUser, AccountIdInvalid, offset, limit, UE::Online::Errors::InvalidUser(), LeaderboardEntries); });
 
 	RunToCompletion();
 }
@@ -67,8 +67,8 @@ LEADERBOARDS_TEST_CASE("Verify ReadEntriesAround User returns entries centered a
 	uint32 limit = 2;
 
 	LeaderboardsFixture_NUsers(MoveTemp(GetLoginPipeline({ LocalUser, OtherPlayerB, UserPlayer, OtherPlayerC, OtherPlayerD })), AccountIds, 5)
-		.EmplaceAsyncLambda([&](FAsyncLambdaResult Promise, SubsystemType Services) { ReadEntriesAroundUser_5UsersFixture(Promise, Services, "Stat_Use_Set", { LocalUser, OtherPlayerB, UserPlayer, OtherPlayerC, OtherPlayerD }, offset, limit, LeaderboardEntries); })
-		.EmplaceLambda([&](SubsystemType Services) { CheckEntries_2UsersRanks_Stat_Use_Set(LeaderboardEntries, OtherPlayerC, OtherPlayerD); });
+		.EmplaceAsyncLambda([&](const FAsyncLambdaResult& Promise, const IOnlineServicesPtr& Services) { ReadEntriesAroundUser_5UsersFixture(Promise, Services, "Stat_Use_Set", { LocalUser, OtherPlayerB, UserPlayer, OtherPlayerC, OtherPlayerD }, offset, limit, LeaderboardEntries); })
+		.EmplaceLambda([&](const IOnlineServicesPtr& Services) { CheckEntries_2UsersRanks_Stat_Use_Set(LeaderboardEntries, OtherPlayerC, OtherPlayerD); });
 
 	RunToCompletion();
 }
@@ -82,9 +82,10 @@ LEADERBOARDS_TEST_CASE("Verify ReadEntriesAroundUser returns a number of entries
 	uint32 limit = 2;
 
 	LeaderboardsFixture_NUsers(MoveTemp(GetLoginPipeline({ LocalUser, AccountId })), AccountIds, 2)
-		.EmplaceAsyncLambda([&](FAsyncLambdaResult Promise, SubsystemType Services) {
-		Services->GetLeaderboardsInterface()->ReadEntriesAroundUser({ LocalUser, AccountId, offset, limit, "Stat_Use_Set" })
-			.OnComplete([Promise, limit](const TOnlineResult<UE::Online::FReadEntriesAroundUser>& Result) mutable
+		.EmplaceAsyncLambda([&](FAsyncLambdaResult Promise, const IOnlineServicesPtr& Services) 
+		{
+			Services->GetLeaderboardsInterface()->ReadEntriesAroundUser({ LocalUser, AccountId, offset, limit, "Stat_Use_Set" })
+			.OnComplete([Promise, limit](const TOnlineResult<UE::Online::FReadEntriesAroundUser>& Result) mutable 
 				{
 					CHECK_EQUAL(Result.GetOkValue().Entries.Num(), limit);
 					Promise->SetValue(true);
@@ -105,8 +106,8 @@ LEADERBOARDS_TEST_CASE("Verify ReadEntriesAroundUser returns expected valid entr
 	TArray<UE::Online::FLeaderboardEntry> LeaderboardEntries;
 
 	LeaderboardsFixture_NUsers(MoveTemp(GetLoginPipeline({ AccountId, LocalUser })), AccountIds, 2)
-		.EmplaceAsyncLambda([&](FAsyncLambdaResult Promise, SubsystemType Services) { ReadEntriesAroundUser_2UsersFixture(Promise, Services, "Stat_Use_Set", LocalUser, AccountId, offset, limit, LeaderboardEntries); })
-		.EmplaceLambda([&](SubsystemType Services) { CheckEntries_2Users_Stat_Use_Set(LeaderboardEntries, AccountId, LocalUser); });
+		.EmplaceAsyncLambda([&](const FAsyncLambdaResult& Promise, const IOnlineServicesPtr& Services) { ReadEntriesAroundUser_2UsersFixture(Promise, Services, "Stat_Use_Set", LocalUser, AccountId, offset, limit, LeaderboardEntries); })
+		.EmplaceLambda([&](const IOnlineServicesPtr& Services) { CheckEntries_2Users_Stat_Use_Set(LeaderboardEntries, AccountId, LocalUser); });
 
 	RunToCompletion();
 }
@@ -122,8 +123,8 @@ LEADERBOARDS_TEST_CASE("Verify ReadEntriesAroundUser returns expected valid entr
 	TArray<UE::Online::FLeaderboardEntry> LeaderboardEntries;
 
 	LeaderboardsFixture_NUsers(MoveTemp(GetLoginPipeline({ LocalUser, AccountId })), AccountIds, 2)
-		.EmplaceAsyncLambda([&](FAsyncLambdaResult Promise, SubsystemType Services) { ReadEntriesAroundUser_2UsersFixture(Promise, Services, "Stat_Use_Smallest", LocalUser, AccountId, offset, limit, LeaderboardEntries); })
-		.EmplaceLambda([&](SubsystemType Services) { CheckEntries_1User_Stat_Use_Smallest(LeaderboardEntries, AccountId); });
+		.EmplaceAsyncLambda([&](const FAsyncLambdaResult& Promise, const IOnlineServicesPtr& Services) { ReadEntriesAroundUser_2UsersFixture(Promise, Services, "Stat_Use_Smallest", LocalUser, AccountId, offset, limit, LeaderboardEntries); })
+		.EmplaceLambda([&](const IOnlineServicesPtr& Services) { CheckEntries_1User_Stat_Use_Smallest(LeaderboardEntries, AccountId); });
 
 	RunToCompletion();
 }
@@ -137,7 +138,7 @@ LEADERBOARDS_TEST_CASE("Verify ReadEntriesAroundUser returns error message if gi
 	TArray<UE::Online::FLeaderboardEntry> LeaderboardEntries;
 
 	LeaderboardsFixture_NUsers(MoveTemp(GetLoginPipeline({ LocalUser, AccountId })), AccountIds, 2)
-		.EmplaceAsyncLambda([&](FAsyncLambdaResult Promise, SubsystemType Services) {ReadEntriesAroundUserFixture(Promise, Services, "Stat_Use_Set", LocalUser, AccountId, offset, limit, UE::Online::Errors::InvalidParams(), LeaderboardEntries); });
+		.EmplaceAsyncLambda([&](const FAsyncLambdaResult& Promise, const IOnlineServicesPtr& Services) {ReadEntriesAroundUserFixture(Promise, Services, "Stat_Use_Set", LocalUser, AccountId, offset, limit, UE::Online::Errors::InvalidParams(), LeaderboardEntries); });
 
 	RunToCompletion();
 }
@@ -150,9 +151,10 @@ LEADERBOARDS_TEST_CASE("Verify ReadEntriesAroundUser returns only one entry if g
 	uint32 limit = 1;
 
 	LeaderboardsFixture_NUsers(MoveTemp(GetLoginPipeline({ LocalUser, AccountId })), AccountIds, 2)
-		.EmplaceAsyncLambda([&](FAsyncLambdaResult Promise, SubsystemType Services) {
-		Services->GetLeaderboardsInterface()->ReadEntriesAroundUser({ LocalUser, AccountId, offset, limit, "Stat_Use_Set" })
-			.OnComplete([Promise, limit](const TOnlineResult<UE::Online::FReadEntriesAroundUser>& Result) mutable
+		.EmplaceAsyncLambda([&](FAsyncLambdaResult Promise, const IOnlineServicesPtr& Services) 
+			{
+			Services->GetLeaderboardsInterface()->ReadEntriesAroundUser({ LocalUser, AccountId, offset, limit, "Stat_Use_Set" })
+			.OnComplete([Promise, limit](const TOnlineResult<UE::Online::FReadEntriesAroundUser>& Result) mutable 
 				{
 					CHECK_EQUAL(Result.GetOkValue().Entries.Num(), limit);
 					Promise->SetValue(true);
@@ -171,8 +173,8 @@ LEADERBOARDS_TEST_CASE("Verify ReadEntriesAroundUser returns expected valid entr
 	uint32 limit = 2;
 
 	LeaderboardsFixture_NUsers(MoveTemp(GetLoginPipeline({ OtherPlayerA, OtherPlayerB, UserPlayer, OtherPlayerC, OtherPlayerD })), AccountIds, 5)
-		.EmplaceAsyncLambda([&](FAsyncLambdaResult Promise, SubsystemType Services) { ReadEntriesAroundUser_5UsersFixture(Promise, Services, "Stat_Use_Set", { OtherPlayerA, OtherPlayerB, UserPlayer, OtherPlayerC, OtherPlayerD }, offset, limit, LeaderboardEntries); })
-		.EmplaceLambda([&](SubsystemType Services) { CheckEntries_2UsersRanks_Stat_Use_Set_Index_Moved_Up(LeaderboardEntries, UserPlayer, OtherPlayerC); });
+		.EmplaceAsyncLambda([&](const FAsyncLambdaResult& Promise, const IOnlineServicesPtr& Services) { ReadEntriesAroundUser_5UsersFixture(Promise, Services, "Stat_Use_Set", { OtherPlayerA, OtherPlayerB, UserPlayer, OtherPlayerC, OtherPlayerD }, offset, limit, LeaderboardEntries); })
+		.EmplaceLambda([&](const IOnlineServicesPtr& Services) { CheckEntries_2UsersRanks_Stat_Use_Set_Index_Moved_Up(LeaderboardEntries, UserPlayer, OtherPlayerC); });
 
 	RunToCompletion();
 }
diff --git a/Engine/Source/Programs/Online/OnlineTestsCore/Source/Private/Tests/Leaderboards/LeaderboardsReadEntriesForUsersTests.cpp b/Engine/Source/Programs/Online/OnlineTestsCore/Source/Private/Tests/Leaderboards/LeaderboardsReadEntriesForUsersTests.cpp
--- a/Engine/Source/Programs/Online/OnlineTestsCore/Source/Private/Tests/Leaderboards/LeaderboardsReadEntriesForUsersTests.cpp
+++ b/Engine/Source/Programs/Online/OnlineTestsCore/Source/Private/Tests/Leaderboards/LeaderboardsReadEntriesForUsersTests.cpp
@@ -18,10 +18,10 @@ LEADERBOARDS_TEST_CASE("ReadEntriesForUsers succeed", "[MultiAccount]")
 	TArray< FAccountId*> AccountIds = { &AccountIdA, &AccountIdB };
 
 	LeaderboardsFixture_NUsers(MoveTemp(GetLoginPipeline({ AccountIdA, AccountIdB })), AccountIds, 2)
-		.EmplaceAsyncLambda([&](FAsyncLambdaResult Promise, SubsystemType Services) { ReadEntriesForUsers_2UsersFixture(Promise, Services, "Stat_Use_Set", AccountIdA, AccountIdB, LeaderboardEntries); })
-		.EmplaceLambda([&](SubsystemType Services) { CheckEntries_2Users_Stat_Use_Set(LeaderboardEntries, AccountIdA, AccountIdB); })
-		.EmplaceAsyncLambda([&](FAsyncLambdaResult Promise, SubsystemType Services) { ReadEntriesForUsers_2UsersFixture(Promise, Services, "Stat_Use_Smallest", AccountIdA, AccountIdB, LeaderboardEntries); })
-		.EmplaceLambda([&](SubsystemType Services) { CheckEntries_2Users_Stat_Use_Smallest(LeaderboardEntries, AccountIdA, AccountIdB); });
+		.EmplaceAsyncLambda([&](const FAsyncLambdaResult& Promise, const IOnlineServicesPtr& Services) { ReadEntriesForUsers_2UsersFixture(Promise, Services, "Stat_Use_Set", AccountIdA, AccountIdB, LeaderboardEntries); })
+		.EmplaceLambda([&](const IOnlineServicesPtr& Services) { CheckEntries_2Users_Stat_Use_Set(LeaderboardEntries, AccountIdA, AccountIdB); })
+		.EmplaceAsyncLambda([&](const FAsyncLambdaResult& Promise, const IOnlineServicesPtr& Services) { ReadEntriesForUsers_2UsersFixture(Promise, Services, "Stat_Use_Smallest", AccountIdA, AccountIdB, LeaderboardEntries); })
+		.EmplaceLambda([&](const IOnlineServicesPtr& Services) { CheckEntries_2Users_Stat_Use_Smallest(LeaderboardEntries, AccountIdA, AccountIdB); });
 
 	RunToCompletion();
 }
@@ -33,7 +33,7 @@ LEADERBOARDS_TEST_CASE("Verify ReadEntriesForUsers returns a fail message if the
 	FAccountId LocalUserInvalid;
 	TArray< FAccountId*> AccountIds = { &AccountId, &LocalUserInvalid };
 	LeaderboardsFixture_NUsers(MoveTemp(GetLoginPipeline({ AccountId })), AccountIds, 2)
-		.EmplaceAsyncLambda([&](FAsyncLambdaResult Promise, SubsystemType Services) { ReadEntriesForUsersFixture(Promise, Services, "Stat_Use_Set", { LocalUserInvalid, AccountId }, UE::Online::Errors::InvalidUser(), LeaderboardEntries); });
+		.EmplaceAsyncLambda([&](const FAsyncLambdaResult& Promise, const IOnlineServicesPtr& Services) { ReadEntriesForUsersFixture(Promise, Services, "Stat_Use_Set", { LocalUserInvalid, AccountId }, UE::Online::Errors::InvalidUser(), LeaderboardEntries); });
 
 	RunToCompletion();
 }
@@ -44,7 +44,7 @@ LEADERBOARDS_TEST_CASE("Verify ReadEntriesForUsers returns an empty array of ent
 	FAccountId LocalUser;
 	TArray< FAccountId*> AccountIds = { &LocalUser };
 	LeaderboardsFixture_NUsers(MoveTemp(GetLoginPipeline({ LocalUser })), AccountIds, 1)
-		.EmplaceAsyncLambda([&](FAsyncLambdaResult Promise, SubsystemType Services) {
+		.EmplaceAsyncLambda([&](const FAsyncLambdaResult& Promise, const IOnlineServicesPtr& Services) {
 		Services->GetLeaderboardsInterface()->ReadEntriesForUsers({ LocalUser, {}, "Stat_Use_Set" })
 		.OnComplete([Promise, &LeaderboardEntries](const TOnlineResult<UE::Online::FReadEntriesForUsers>& Result) mutable
 			{
@@ -62,9 +62,10 @@ LEADERBOARDS_TEST_CASE("Verify ReadEntiresForUsers returns a fail message if giv
 	FAccountId LocalUser;
 	TArray< FAccountId*> AccountIds = { &LocalUser, &AccountIdInvalid };
 	LeaderboardsFixture_NUsers(MoveTemp(GetLoginPipeline({ LocalUser })), AccountIds, 1)
-		.EmplaceAsyncLambda([&](FAsyncLambdaResult Promise, SubsystemType Services) {
-		Services->GetLeaderboardsInterface()->ReadEntriesForUsers({ LocalUser, {AccountIdInvalid}, "Stat_Use_Set" })
-		.OnComplete([Promise, &LeaderboardEntries](const TOnlineResult<UE::Online::FReadEntriesForUsers>& Result) mutable
+		.EmplaceAsyncLambda([&](FAsyncLambdaResult Promise, const IOnlineServicesPtr& Services) 
+		{
+			Services->GetLeaderboardsInterface()->ReadEntriesForUsers({ LocalUser, {AccountIdInvalid}, "Stat_Use_Set" })
+			.OnComplete([Promise, &LeaderboardEntries](const TOnlineResult<UE::Online::FReadEntriesForUsers>& Result) mutable 
 			{
 				if (Result.IsError())
 				{
@@ -88,7 +89,7 @@ LEADERBOARDS_TEST_CASE("Verify ReadEntriesForUsers returns a fail message if giv
 	FAccountId LocalUser;
 	TArray< FAccountId*> AccountIds = { &LocalUser, &AccountIdInvalid };
 	LeaderboardsFixture_NUsers(MoveTemp(GetLoginPipeline({ LocalUser })), AccountIds, 1)
-		.EmplaceAsyncLambda([&](FAsyncLambdaResult Promise, SubsystemType Services) { ReadEntriesForUsersFixture(Promise, Services, "Stat_Use_Set", { LocalUser, AccountIdInvalid }, UE::Online::Errors::InvalidUser(), LeaderboardEntries); });
+		.EmplaceAsyncLambda([&](const FAsyncLambdaResult& Promise, const IOnlineServicesPtr& Services) { ReadEntriesForUsersFixture(Promise, Services, "Stat_Use_Set", { LocalUser, AccountIdInvalid }, UE::Online::Errors::InvalidUser(), LeaderboardEntries); });
 	RunToCompletion();
 }
 
@@ -98,9 +99,10 @@ LEADERBOARDS_TEST_CASE("Verify ReadEntriesForUsers returns an Entries array with
 	FAccountId LocalUser;
 	TArray< FAccountId*> AccountIds = { &LocalUser };
 	LeaderboardsFixture_NUsers(MoveTemp(GetLoginPipeline({ LocalUser })), AccountIds, 1)
-		.EmplaceAsyncLambda([&](FAsyncLambdaResult Promise, SubsystemType Services) {
-		Services->GetLeaderboardsInterface()->ReadEntriesForUsers({ LocalUser, {LocalUser}, "Stat_Use_Set" })
-		.OnComplete([Promise, &LeaderboardEntries, &LocalUser](const TOnlineResult<UE::Online::FReadEntriesForUsers>& Result) mutable
+		.EmplaceAsyncLambda([&](FAsyncLambdaResult Promise, const IOnlineServicesPtr& Services) 
+		{
+			Services->GetLeaderboardsInterface()->ReadEntriesForUsers({ LocalUser, {LocalUser}, "Stat_Use_Set" })
+			.OnComplete([Promise, &LeaderboardEntries, &LocalUser](const TOnlineResult<UE::Online::FReadEntriesForUsers>& Result) mutable 
 			{
 				CHECK_EQUAL(Result.GetOkValue().Entries.Num(), 1);
 				CHECK_EQUAL(LeaderboardEntryOf(Result.GetOkValue().Entries, LocalUser)->Rank, 1);
diff --git a/Engine/Source/Programs/Online/OnlineTestsCore/Source/Private/Tests/Leaderboards/ReadEntiresForUsersFixture.cpp b/Engine/Source/Programs/Online/OnlineTestsCore/Source/Private/Tests/Leaderboards/ReadEntiresForUsersFixture.cpp
--- a/Engine/Source/Programs/Online/OnlineTestsCore/Source/Private/Tests/Leaderboards/ReadEntiresForUsersFixture.cpp
+++ b/Engine/Source/Programs/Online/OnlineTestsCore/Source/Private/Tests/Leaderboards/ReadEntiresForUsersFixture.cpp
@@ -3,7 +3,7 @@
 #include "ReadEntiresForUsersFixture.h"
 
 
-void ReadEntriesForUsers_2UsersFixture(FAsyncLambdaResult Promise, SubsystemType Services, const FString& BoardName, const FAccountId& AccountIdA, const FAccountId& AccountIdB, TArray<UE::Online::FLeaderboardEntry>& OutLeaderboardEntries)
+void ReadEntriesForUsers_2UsersFixture(FAsyncLambdaResult Promise, const IOnlineServicesPtr& Services, const FString& BoardName, const FAccountId& AccountIdA, const FAccountId& AccountIdB, TArray<UE::Online::FLeaderboardEntry>& OutLeaderboardEntries)
 {
 	Services->GetLeaderboardsInterface()->ReadEntriesForUsers({ AccountIdA, {AccountIdA, AccountIdB}, BoardName }).OnComplete([Promise, &OutLeaderboardEntries](const TOnlineResult<UE::Online::FReadEntriesForUsers>& Result) mutable
 		{
@@ -13,7 +13,7 @@ void ReadEntriesForUsers_2UsersFixture(FAsyncLambdaResult Promise, SubsystemType
 		});
 }
 
-void ReadEntriesForUsersFixture(FAsyncLambdaResult Promise, SubsystemType Services, const FString& BoardName, const TArray<FAccountId>& AccountIds, const UE::Online::FOnlineError& InError, TArray<UE::Online::FLeaderboardEntry>& OutLeaderboardEntries)
+void ReadEntriesForUsersFixture(FAsyncLambdaResult Promise, const IOnlineServicesPtr& Services, const FString& BoardName, const TArray<FAccountId>& AccountIds, const UE::Online::FOnlineError& InError, TArray<UE::Online::FLeaderboardEntry>& OutLeaderboardEntries)
 {
 	Services->GetLeaderboardsInterface()->ReadEntriesForUsers({ AccountIds[0], AccountIds, BoardName }).OnComplete([Promise, &OutLeaderboardEntries, InError](const TOnlineResult<UE::Online::FReadEntriesForUsers>& Result) mutable
 		{
diff --git a/Engine/Source/Programs/Online/OnlineTestsCore/Source/Private/Tests/Leaderboards/ReadEntiresForUsersFixture.h b/Engine/Source/Programs/Online/OnlineTestsCore/Source/Private/Tests/Leaderboards/ReadEntiresForUsersFixture.h
--- a/Engine/Source/Programs/Online/OnlineTestsCore/Source/Private/Tests/Leaderboards/ReadEntiresForUsersFixture.h
+++ b/Engine/Source/Programs/Online/OnlineTestsCore/Source/Private/Tests/Leaderboards/ReadEntiresForUsersFixture.h
@@ -10,6 +10,6 @@
 #include "Online/Leaderboards.h"
 
 
-void ReadEntriesForUsers_2UsersFixture(FAsyncLambdaResult Promise, SubsystemType Services, const FString& BoardName, const FAccountId& AccountIdA, const FAccountId& AccountIdB, TArray<UE::Online::FLeaderboardEntry>& OutLeaderboardEntries);
-void ReadEntriesForUsersFixture(FAsyncLambdaResult Promise, SubsystemType Services, const FString& BoardName, const TArray<FAccountId>& AccountIds, const UE::Online::FOnlineError& InError, TArray<UE::Online::FLeaderboardEntry>& OutLeaderboardEntries);
+void ReadEntriesForUsers_2UsersFixture(FAsyncLambdaResult Promise, const IOnlineServicesPtr& Services, const FString& BoardName, const FAccountId& AccountIdA, const FAccountId& AccountIdB, TArray<UE::Online::FLeaderboardEntry>& OutLeaderboardEntries);
+void ReadEntriesForUsersFixture(FAsyncLambdaResult Promise, const IOnlineServicesPtr& Services, const FString& BoardName, const TArray<FAccountId>& AccountIds, const UE::Online::FOnlineError& InError, TArray<UE::Online::FLeaderboardEntry>& OutLeaderboardEntries);
 
diff --git a/Engine/Source/Programs/Online/OnlineTestsCore/Source/Private/Tests/Leaderboards/ReadEntriesAroundRankFixture.cpp b/Engine/Source/Programs/Online/OnlineTestsCore/Source/Private/Tests/Leaderboards/ReadEntriesAroundRankFixture.cpp
--- a/Engine/Source/Programs/Online/OnlineTestsCore/Source/Private/Tests/Leaderboards/ReadEntriesAroundRankFixture.cpp
+++ b/Engine/Source/Programs/Online/OnlineTestsCore/Source/Private/Tests/Leaderboards/ReadEntriesAroundRankFixture.cpp
@@ -2,7 +2,7 @@
 
 #include "ReadEntriesAroundRankFixture.h"
 
-void ReadEntriesAroundRank_2UsersFixture(FAsyncLambdaResult Promise, SubsystemType Services, const FString& BoardName, const FAccountId& AccountIdA, uint32 Rank, uint32 Limit, TArray<UE::Online::FLeaderboardEntry>& OutLeaderboardEntries)
+void ReadEntriesAroundRank_2UsersFixture(FAsyncLambdaResult Promise, const IOnlineServicesPtr& Services, const FString& BoardName, const FAccountId& AccountIdA, uint32 Rank, uint32 Limit, TArray<UE::Online::FLeaderboardEntry>& OutLeaderboardEntries)
 {
 	Services->GetLeaderboardsInterface()->ReadEntriesAroundRank({ AccountIdA, Rank, Limit, BoardName }).OnComplete([Promise, &OutLeaderboardEntries](const TOnlineResult<UE::Online::FReadEntriesAroundRank>& Result) mutable
 		{
@@ -12,7 +12,7 @@ void ReadEntriesAroundRank_2UsersFixture(FAsyncLambdaResult Promise, SubsystemTy
 		});
 }
 
-void ReadEntriesAroundRankFixture(FAsyncLambdaResult Promise, SubsystemType Services, const FString& BoardName, const FAccountId& AccountId, uint32 Rank, uint32 Limit, const UE::Online::FOnlineError& InError, TArray<UE::Online::FLeaderboardEntry>& OutLeaderboardEntries)
+void ReadEntriesAroundRankFixture(FAsyncLambdaResult Promise, const IOnlineServicesPtr& Services, const FString& BoardName, const FAccountId& AccountId, uint32 Rank, uint32 Limit, const UE::Online::FOnlineError& InError, TArray<UE::Online::FLeaderboardEntry>& OutLeaderboardEntries)
 {
 	Services->GetLeaderboardsInterface()->ReadEntriesAroundRank({ AccountId, Rank, Limit, BoardName }).OnComplete([Promise, &OutLeaderboardEntries, InError](const TOnlineResult<UE::Online::FReadEntriesAroundRank>& Result) mutable
 		{
diff --git a/Engine/Source/Programs/Online/OnlineTestsCore/Source/Private/Tests/Leaderboards/ReadEntriesAroundRankFixture.h b/Engine/Source/Programs/Online/OnlineTestsCore/Source/Private/Tests/Leaderboards/ReadEntriesAroundRankFixture.h
--- a/Engine/Source/Programs/Online/OnlineTestsCore/Source/Private/Tests/Leaderboards/ReadEntriesAroundRankFixture.h
+++ b/Engine/Source/Programs/Online/OnlineTestsCore/Source/Private/Tests/Leaderboards/ReadEntriesAroundRankFixture.h
@@ -10,5 +10,5 @@
 #include "Online/Leaderboards.h"
 
 
-void ReadEntriesAroundRank_2UsersFixture(FAsyncLambdaResult Promise, SubsystemType Services, const FString& BoardName, const FAccountId& AccountIdA, uint32 Rank, uint32 Limit, TArray<UE::Online::FLeaderboardEntry>& OutLeaderboardEntries);
-void ReadEntriesAroundRankFixture(FAsyncLambdaResult Promise, SubsystemType Services, const FString& BoardName, const FAccountId& AccountId, uint32 Rank, uint32 Limit, const UE::Online::FOnlineError& InError, TArray<UE::Online::FLeaderboardEntry>& OutLeaderboardEntries);
+void ReadEntriesAroundRank_2UsersFixture(FAsyncLambdaResult Promise, const IOnlineServicesPtr& Services, const FString& BoardName, const FAccountId& AccountIdA, uint32 Rank, uint32 Limit, TArray<UE::Online::FLeaderboardEntry>& OutLeaderboardEntries);
+void ReadEntriesAroundRankFixture(FAsyncLambdaResult Promise, const IOnlineServicesPtr& Services, const FString& BoardName, const FAccountId& AccountId, uint32 Rank, uint32 Limit, const UE::Online::FOnlineError& InError, TArray<UE::Online::FLeaderboardEntry>& OutLeaderboardEntries);
diff --git a/Engine/Source/Programs/Online/OnlineTestsCore/Source/Private/Tests/Leaderboards/ReadEntriesAroundUserFixture.cpp b/Engine/Source/Programs/Online/OnlineTestsCore/Source/Private/Tests/Leaderboards/ReadEntriesAroundUserFixture.cpp
--- a/Engine/Source/Programs/Online/OnlineTestsCore/Source/Private/Tests/Leaderboards/ReadEntriesAroundUserFixture.cpp
+++ b/Engine/Source/Programs/Online/OnlineTestsCore/Source/Private/Tests/Leaderboards/ReadEntriesAroundUserFixture.cpp
@@ -2,7 +2,7 @@
 
 #include "ReadEntriesAroundUserFixture.h"
 
-void ReadEntriesAroundUser_2UsersFixture(FAsyncLambdaResult Promise, SubsystemType Services, const FString& BoardName, const FAccountId& AccountIdA, const FAccountId& AccountIdB, int32 Offset, uint32 Limit, TArray<UE::Online::FLeaderboardEntry>& OutLeaderboardEntries)
+void ReadEntriesAroundUser_2UsersFixture(FAsyncLambdaResult Promise, const IOnlineServicesPtr& Services, const FString& BoardName, const FAccountId& AccountIdA, const FAccountId& AccountIdB, int32 Offset, uint32 Limit, TArray<UE::Online::FLeaderboardEntry>& OutLeaderboardEntries)
 {
 	Services->GetLeaderboardsInterface()->ReadEntriesAroundUser({ AccountIdA, AccountIdB, Offset, Limit, BoardName }).OnComplete([Promise, &OutLeaderboardEntries](const TOnlineResult<UE::Online::FReadEntriesAroundUser>& Result) mutable
 		{
@@ -12,7 +12,7 @@ void ReadEntriesAroundUser_2UsersFixture(FAsyncLambdaResult Promise, SubsystemTy
 		});
 }
 
-void ReadEntriesAroundUser_5UsersFixture(FAsyncLambdaResult Promise, SubsystemType Services, const FString& BoardName, const TArray<FAccountId>& AccountIds, int32 Offset, uint32 Limit, TArray<UE::Online::FLeaderboardEntry>& OutLeaderboardEntries)
+void ReadEntriesAroundUser_5UsersFixture(FAsyncLambdaResult Promise, const IOnlineServicesPtr& Services, const FString& BoardName, const TArray<FAccountId>& AccountIds, int32 Offset, uint32 Limit, TArray<UE::Online::FLeaderboardEntry>& OutLeaderboardEntries)
 {
 	Services->GetLeaderboardsInterface()->ReadEntriesAroundUser({ AccountIds[0], AccountIds[2], Offset, Limit, BoardName }).OnComplete([Promise, &OutLeaderboardEntries](const TOnlineResult<UE::Online::FReadEntriesAroundUser>& Result) mutable
 		{
@@ -22,7 +22,7 @@ void ReadEntriesAroundUser_5UsersFixture(FAsyncLambdaResult Promise, SubsystemTy
 		});
 }
 
-void ReadEntriesAroundUserFixture(FAsyncLambdaResult Promise, SubsystemType Services, const FString& BoardName, const FAccountId& AccountIdA, const FAccountId& AccountIdB, int32 Offset, uint32 Limit, const UE::Online::FOnlineError& InError, TArray<UE::Online::FLeaderboardEntry>& OutLeaderboardEntries)
+void ReadEntriesAroundUserFixture(FAsyncLambdaResult Promise, const IOnlineServicesPtr& Services, const FString& BoardName, const FAccountId& AccountIdA, const FAccountId& AccountIdB, int32 Offset, uint32 Limit, const UE::Online::FOnlineError& InError, TArray<UE::Online::FLeaderboardEntry>& OutLeaderboardEntries)
 {
 	Services->GetLeaderboardsInterface()->ReadEntriesAroundUser({ AccountIdA, AccountIdB, Offset, Limit, BoardName }).OnComplete([Promise, &OutLeaderboardEntries, InError](const TOnlineResult<UE::Online::FReadEntriesAroundUser>& Result) mutable
 		{
diff --git a/Engine/Source/Programs/Online/OnlineTestsCore/Source/Private/Tests/Leaderboards/ReadEntriesAroundUserFixture.h b/Engine/Source/Programs/Online/OnlineTestsCore/Source/Private/Tests/Leaderboards/ReadEntriesAroundUserFixture.h
--- a/Engine/Source/Programs/Online/OnlineTestsCore/Source/Private/Tests/Leaderboards/ReadEntriesAroundUserFixture.h
+++ b/Engine/Source/Programs/Online/OnlineTestsCore/Source/Private/Tests/Leaderboards/ReadEntriesAroundUserFixture.h
@@ -10,6 +10,6 @@
 #include "Online/Leaderboards.h"
 
 
-void ReadEntriesAroundUser_2UsersFixture(FAsyncLambdaResult Promise, SubsystemType Services, const FString& BoardName, const FAccountId& AccountIdA, const FAccountId& AccountIdB, int32 Offset, uint32 Limit, TArray<UE::Online::FLeaderboardEntry>& OutLeaderboardEntries);
-void ReadEntriesAroundUser_5UsersFixture(FAsyncLambdaResult Promise, SubsystemType Services, const FString& BoardName, const TArray<FAccountId>& AccountIds, int32 Offset, uint32 Limit, TArray<UE::Online::FLeaderboardEntry>& OutLeaderboardEntries);
-void ReadEntriesAroundUserFixture(FAsyncLambdaResult Promise, SubsystemType Services, const FString& BoardName, const FAccountId& AccountIdA, const FAccountId& AccountIdB, int32 Offset, uint32 Limit, const UE::Online::FOnlineError& InError, TArray<UE::Online::FLeaderboardEntry>& OutLeaderboardEntries);
+void ReadEntriesAroundUser_2UsersFixture(FAsyncLambdaResult Promise, const IOnlineServicesPtr& Services, const FString& BoardName, const FAccountId& AccountIdA, const FAccountId& AccountIdB, int32 Offset, uint32 Limit, TArray<UE::Online::FLeaderboardEntry>& OutLeaderboardEntries);
+void ReadEntriesAroundUser_5UsersFixture(FAsyncLambdaResult Promise, const IOnlineServicesPtr& Services, const FString& BoardName, const TArray<FAccountId>& AccountIds, int32 Offset, uint32 Limit, TArray<UE::Online::FLeaderboardEntry>& OutLeaderboardEntries);
+void ReadEntriesAroundUserFixture(FAsyncLambdaResult Promise, const IOnlineServicesPtr& Services, const FString& BoardName, const FAccountId& AccountIdA, const FAccountId& AccountIdB, int32 Offset, uint32 Limit, const UE::Online::FOnlineError& InError, TArray<UE::Online::FLeaderboardEntry>& OutLeaderboardEntries);
diff --git a/Engine/Source/Programs/Online/OnlineTestsCore/Source/Private/Tests/Meta/AsyncTestStepTest.cpp b/Engine/Source/Programs/Online/OnlineTestsCore/Source/Private/Tests/Meta/AsyncTestStepTest.cpp
--- a/Engine/Source/Programs/Online/OnlineTestsCore/Source/Private/Tests/Meta/AsyncTestStepTest.cpp
+++ b/Engine/Source/Programs/Online/OnlineTestsCore/Source/Private/Tests/Meta/AsyncTestStepTest.cpp
@@ -15,7 +15,7 @@ public:
 
 	}
 
-	virtual void Run(FAsyncStepResult Promise, SubsystemType Services) override
+	virtual void Run(FAsyncStepResult Promise, const IOnlineServicesPtr& Services) override
 	{
 		CHECK(Services != nullptr);
 		bHasRun = true;
@@ -35,16 +35,16 @@ AST_ERROR_TEST_CASE("Basic test- confirm all of async steps, lambdas, and async
 
 	GetPipeline()
 		.EmplaceStep<FAsyncTestStepHelper>(bHasRun)
-		.EmplaceLambda([&](SubsystemType Type)
+		.EmplaceLambda([&](const IOnlineServicesPtr& Type)
 		{
 			CHECK(bHasRun);
 		})
-		.EmplaceAsyncLambda([&](FAsyncLambdaResult Result, SubsystemType Type)
+		.EmplaceAsyncLambda([&](const FAsyncLambdaResult& Result, const IOnlineServicesPtr& Type)
 		{
 			CHECK(bHasRun);
 			Result->SetValue(true);
 		})
-		.EmplaceLambda([&](SubsystemType Type)
+		.EmplaceLambda([&](const IOnlineServicesPtr& Type)
 		{
 			bDidComplete = true;
 		});
diff --git a/Engine/Source/Programs/Online/OnlineTestsCore/Source/Private/Tests/Meta/MetaTests.cpp b/Engine/Source/Programs/Online/OnlineTestsCore/Source/Private/Tests/Meta/MetaTests.cpp
--- a/Engine/Source/Programs/Online/OnlineTestsCore/Source/Private/Tests/Meta/MetaTests.cpp
+++ b/Engine/Source/Programs/Online/OnlineTestsCore/Source/Private/Tests/Meta/MetaTests.cpp
@@ -17,13 +17,13 @@ META_TEST_CASE("Async test steps")
 		bool bDidComplete = false;
 
 		GetPipeline()
-			.EmplaceAsyncLambda([&](FAsyncLambdaResult Result, SubsystemType Type)
+			.EmplaceAsyncLambda([&](FAsyncLambdaResult Result, const IOnlineServicesPtr& Type)
 				{
 					// This should cause a failure and crash, but if the fix is working
 					// correctly, it should not prevent the promise from being fulfilled.
 					REQUIRE(false);
 				})
-			.EmplaceLambda([&](SubsystemType Type)
+			.EmplaceLambda([&](const IOnlineServicesPtr& Type)
 				{
 					// Only runs if the previous step completed (the promise was fulfilled)
 					bDidComplete = true;
diff --git a/Engine/Source/Programs/Online/OnlineTestsCore/Source/Private/Tests/Meta/SyncLoginTest.cpp b/Engine/Source/Programs/Online/OnlineTestsCore/Source/Private/Tests/Meta/SyncLoginTest.cpp
--- a/Engine/Source/Programs/Online/OnlineTestsCore/Source/Private/Tests/Meta/SyncLoginTest.cpp
+++ b/Engine/Source/Programs/Online/OnlineTestsCore/Source/Private/Tests/Meta/SyncLoginTest.cpp
@@ -17,13 +17,13 @@ SYNCLOGIN_ERROR_TEST_CASE("Confirm sync login is logging in before RunToCompleti
 	bool bHasRun = false;
 	
 	PipeSave
-		.EmplaceLambda([&](SubsystemType Type)
+		.EmplaceLambda([&](const IOnlineServicesPtr& Type)
 			{
 				// Check valid during pipeline
 				CHECK(AccountId.IsValid());
 				CHECK(AccountId2.IsValid());
 			})
-		.EmplaceLambda([&](SubsystemType Type)
+		.EmplaceLambda([&](const IOnlineServicesPtr& Type)
 		{
 			bHasRun = true;
 		});
@@ -38,7 +38,7 @@ SYNCLOGIN_ERROR_TEST_CASE("Confirm sync login is logging in before RunToCompleti
 	bool bHasRun = false;
 
 	GetLoginPipeline({ AccountId, AccountId2, AccountId3, AccountId4, AccountId5 })
-		.EmplaceLambda([&](SubsystemType Type)
+		.EmplaceLambda([&](const IOnlineServicesPtr& Type)
 		{
 				// Check valid during pipeline
 				CHECK(AccountId.IsValid());
@@ -47,7 +47,7 @@ SYNCLOGIN_ERROR_TEST_CASE("Confirm sync login is logging in before RunToCompleti
 				CHECK(AccountId4.IsValid());
 				CHECK(AccountId5.IsValid());
 		})
-		.EmplaceLambda([&bHasRun](SubsystemType Type)
+		.EmplaceLambda([&bHasRun](const IOnlineServicesPtr& Type)
 		{
 			bHasRun = true;
 		});
diff --git a/Engine/Source/Programs/Online/OnlineTestsCore/Source/Private/Tests/Presence/PresenceTests.cpp b/Engine/Source/Programs/Online/OnlineTestsCore/Source/Private/Tests/Presence/PresenceTests.cpp
--- a/Engine/Source/Programs/Online/OnlineTestsCore/Source/Private/Tests/Presence/PresenceTests.cpp
+++ b/Engine/Source/Programs/Online/OnlineTestsCore/Source/Private/Tests/Presence/PresenceTests.cpp
@@ -45,7 +45,7 @@ PRESENCE_TEST_CASE("Ensure UpdatePresence works on multiple users and can then b
 	TArray<TSharedRef<const FUserPresence>> OutPresences;
 
 	GetLoginPipeline({ AccountIdA, AccountIdB })
-		.EmplaceLambda([&](SubsystemType Type)
+		.EmplaceLambda([&](const IOnlineServicesPtr& Type)
 		{
 			AccountIds.Add(AccountIdA);
 			AccountIds.Add(AccountIdB);
@@ -53,7 +53,7 @@ PRESENCE_TEST_CASE("Ensure UpdatePresence works on multiple users and can then b
 		.EmplaceStep<FUpdatePresenceHelper>(AccountIdA, GetGenericPresenceA(AccountIdA))
 		.EmplaceStep<FUpdatePresenceHelper>(AccountIdB, GetGenericPresenceB(AccountIdB))
 		.EmplaceStep<FBatchQueryPresenceHelper>(AccountIdA, AccountIds, OutPresences)
-		.EmplaceLambda([&](SubsystemType Type)
+		.EmplaceLambda([&](const IOnlineServicesPtr& Type)
 		{
 			REQUIRE(OutPresences.Num() == 2);
 			CheckPresenceAreEqual(*OutPresences[0], *GetGenericPresenceA(AccountIdA));
@@ -88,7 +88,7 @@ PRESENCE_TEST_CASE("Ensure UpdatePresence works as expected and can then be quer
 		.EmplaceStep<FUpdatePresenceHelper>(AccountIdA, GetGenericPresenceA(AccountIdA))
 		.EmplaceStep<FQueryPresenceHelper>(AccountIdA, AccountIdA, OutPresenceA)
 		.EmplaceStep<FGetCachedPresenceHelper>(AccountIdA, AccountIdA, OutPresenceB)
-		.EmplaceLambda([&](SubsystemType Type)
+		.EmplaceLambda([&](const IOnlineServicesPtr& Type)
 		{
 			CheckPresenceAreEqual(*OutPresenceA, *OutPresenceB);
 		});
@@ -118,7 +118,7 @@ PRESENCE_TEST_CASE("Ensure QueryPresence's bListenToChanges mod fires the event
 		.EmplaceStep<FQueryPresenceHelper>(AccountIdA, AccountIdB, OutPresence, true)
 		.EmplaceStep<FUpdatePresenceHelper>(AccountIdB, GetGenericPresenceB(AccountIdB))
 		// todo: possibly add a fixed delay here or a waited retry on the event
-		.EmplaceLambda([&](SubsystemType Type)
+		.EmplaceLambda([&](const IOnlineServicesPtr& Type)
 		{
 			CHECK(bReceivedPresence);
 		});
@@ -146,7 +146,7 @@ PRESENCE_TEST_CASE("Ensure BatchQueryPresence's bListenToChanges mod fires the e
 	});
 
 	GetLoginPipeline({ AccountIdA, AccountIdB })
-		.EmplaceLambda([&](SubsystemType Type)
+		.EmplaceLambda([&](const IOnlineServicesPtr& Type)
 		{
 			AccountIds.Add(AccountIdA);
 			AccountIds.Add(AccountIdB);
@@ -157,7 +157,7 @@ PRESENCE_TEST_CASE("Ensure BatchQueryPresence's bListenToChanges mod fires the e
 		.EmplaceStep<FUpdatePresenceHelper>(AccountIdA, GetGenericPresenceB(AccountIdA))
 		.EmplaceStep<FUpdatePresenceHelper>(AccountIdB, GetGenericPresenceB(AccountIdB))
 		// todo: possibly add a fixed delay here or a waited retry on the event
-		.EmplaceLambda([&](SubsystemType Type)
+		.EmplaceLambda([&](const IOnlineServicesPtr& Type)
 		{
 			CHECK(bReceivedPresence);
 		});
diff --git a/Engine/Source/Programs/Online/OnlineTestsCore/Source/Private/Tests/Sessions/FindSessionsTests.cpp b/Engine/Source/Programs/Online/OnlineTestsCore/Source/Private/Tests/Sessions/FindSessionsTests.cpp
--- a/Engine/Source/Programs/Online/OnlineTestsCore/Source/Private/Tests/Sessions/FindSessionsTests.cpp
+++ b/Engine/Source/Programs/Online/OnlineTestsCore/Source/Private/Tests/Sessions/FindSessionsTests.cpp
@@ -188,9 +188,9 @@ SESSIONS_TEST_CASE("If I call FindSessions while a search is already running for
 	LoginPipeline
 		.EmplaceStep<FCreateSessionHelper>(MoveTemp(CreateSessionHelperParams))
 		.EmplaceStep<FAddSessionMemberHelper>(MoveTemp(AddSessionMemberHelperParams))
-		.EmplaceAsyncLambda([&FindFirstSessionsHelperParams, &FindSecondSessionsHelperParams](FAsyncLambdaResult Promise, SubsystemType OnlineSubsystem)
+		.EmplaceAsyncLambda([&FindFirstSessionsHelperParams, &FindSecondSessionsHelperParams](FAsyncLambdaResult Promise, const IOnlineServicesPtr& OnlineServices)
 			{
-				ISessionsPtr SessionsInterface = OnlineSubsystem->GetSessionsInterface();
+				ISessionsPtr SessionsInterface = OnlineServices->GetSessionsInterface();
 
 				SessionsInterface->FindSessions(MoveTemp(*FindFirstSessionsHelperParams.OpParams))
 					.OnComplete([Promise = MoveTemp(Promise)](const TOnlineResult<FFindSessions>& Result)
@@ -349,8 +349,8 @@ SESSIONS_TEST_CASE("If I call FindSessions with valid information, I will get a
 
 	FirstLeaveSessionHelperParams.OpParams->LocalAccountId = FirstAccountId;
 	SecondLeaveSessionHelperParams.OpParams->LocalAccountId = FirstAccountId;
-	
-	const uint32_t ExpectedSessionsFound = 2;
+
+	constexpr uint32_t ExpectedSessionsFound = 2;
 
 	LoginPipeline
 		.EmplaceStep<FCreateSessionHelper>(MoveTemp(FirstCreateSessionHelperParams))
@@ -425,7 +425,7 @@ SESSIONS_TEST_CASE("Call FindSessions using custom setting as a filter Bool", EG
 	FirstLeaveSessionHelperParams.OpParams->LocalAccountId = FirstAccountId;
 	SecondLeaveSessionHelperParams.OpParams->LocalAccountId = FirstAccountId;
 
-	uint32_t ExpectedSessionsFound = 1;
+	constexpr uint32_t ExpectedSessionsFound = 1;
 
 	LoginPipeline
 	.EmplaceStep<FCreateSessionHelper>(MoveTemp(FirstCreateSessionHelperParams))
@@ -527,7 +527,7 @@ SESSIONS_TEST_CASE("Call FindSessions using custom setting as a filter Int", EG_
 	SeventhFindSessionsHelperParams.OpParams->LocalAccountId = SecondAccountId;
 	EighthFindSessionsHelperParams.OpParams->LocalAccountId = FirstAccountId;
 
-	uint32_t ExpectedSessionsFound = 1;
+	constexpr uint32_t ExpectedSessionsFound = 1;
 
 	LoginPipeline
 		.EmplaceStep<FCreateSessionHelper>(MoveTemp(FirstCreateSessionHelperParams))
@@ -642,7 +642,7 @@ SESSIONS_TEST_CASE("Call FindSessions using custom setting as a filter Double",
 	FourteenthFindSessionsHelperParams.OpParams->LocalAccountId = FirstAccountId;
 	FifteenthFindSessionsHelperParams.OpParams->LocalAccountId = SecondAccountId;
 
-	uint32_t ExpectedSessionsFound = 1;
+	constexpr uint32_t ExpectedSessionsFound = 1;
 
 	LoginPipeline
 		.EmplaceStep<FCreateSessionHelper>(MoveTemp(FirstCreateSessionHelperParams))
@@ -654,9 +654,9 @@ SESSIONS_TEST_CASE("Call FindSessions using custom setting as a filter Double",
 		.EmplaceStep<FFindSessionsHelper>(MoveTemp(EleventhFindSessionsHelperParams), ExpectedSessionsFound)
 		.EmplaceStep<FFindSessionsHelper>(MoveTemp(ThirteenthFindSessionsHelperParams), ExpectedSessionsFound)
 		.EmplaceStep<FFindSessionsHelper>(MoveTemp(FourteenthFindSessionsHelperParams), ExpectedSessionsFound)
-		.EmplaceAsyncLambda([&FifteenthFindSessionsHelperParams, &DoubleSettingValue](FAsyncLambdaResult Promise, SubsystemType OnlineSubsystem)
+		.EmplaceAsyncLambda([&FifteenthFindSessionsHelperParams, &DoubleSettingValue](FAsyncLambdaResult Promise, const IOnlineServicesPtr& OnlineServices)
 			{
-				ISessionsPtr SessionsInterface = OnlineSubsystem->GetSessionsInterface();
+				ISessionsPtr SessionsInterface = OnlineServices->GetSessionsInterface();
 				SessionsInterface->FindSessions(MoveTemp(*FifteenthFindSessionsHelperParams.OpParams))
 					.OnComplete([Promise = MoveTemp(Promise), SessionsInterface, &DoubleSettingValue](const TOnlineResult<FFindSessions>& Result)
 						{
@@ -759,7 +759,7 @@ SESSIONS_TEST_CASE("Call FindSessions using custom setting as a filter String",
 	NineteenthFindSessionsHelperParams.OpParams->LocalAccountId = SecondAccountId;
 
 	TOptional<uint32_t> NoSessionsFound = TOptional<uint32_t>();
-	uint32_t ExpectedSessionsFound = 1;
+	constexpr uint32_t ExpectedSessionsFound = 1;
 
 	LoginPipeline
 		.EmplaceStep<FCreateSessionHelper>(MoveTemp(FirstCreateSessionHelperParams))
diff --git a/Engine/Source/Programs/Online/OnlineTestsCore/Source/Private/Tests/Sessions/GetSessionsTests.cpp b/Engine/Source/Programs/Online/OnlineTestsCore/Source/Private/Tests/Sessions/GetSessionsTests.cpp
--- a/Engine/Source/Programs/Online/OnlineTestsCore/Source/Private/Tests/Sessions/GetSessionsTests.cpp
+++ b/Engine/Source/Programs/Online/OnlineTestsCore/Source/Private/Tests/Sessions/GetSessionsTests.cpp
@@ -16,7 +16,7 @@
 SESSIONS_TEST_CASE("If I call GetAllSessions with an invalid account id, I get an error", EG_SESSIONS_GETSESSIONS_TAG)
 {
 	GetPipeline()
-		.EmplaceLambda([](SubsystemType OnlineSubsystem)
+		.EmplaceLambda([](const IOnlineServicesPtr& OnlineSubsystem)
 			{
 				FGetAllSessions::Params OpGetAllParams;
 				OpGetAllParams.LocalAccountId = FAccountId();
@@ -37,7 +37,7 @@ SESSIONS_TEST_CASE("If I call GetAllSessions before creating or joining any sess
 
 	FAccountId AccountId;
 	GetLoginPipeline({ AccountId })
-		.EmplaceLambda([&AccountId](SubsystemType OnlineSubsystem)
+		.EmplaceLambda([&AccountId](const IOnlineServicesPtr& OnlineSubsystem)
 			{
 				FGetAllSessions::Params OpParams;
 				OpParams.LocalAccountId = AccountId;
@@ -76,11 +76,11 @@ SESSIONS_TEST_CASE("If I call GetAllSessions with valid conditions, I get a vali
 	CreateSessionHelperParams.OpParams->LocalAccountId = AccountId;
 	LeaveSessionHelperParams.OpParams->LocalAccountId = AccountId;
 
-	uint32_t ExpectedSessionsFound = 1;
+	constexpr uint32_t ExpectedSessionsFound = 1;
 
 	LoginPipeline
 		.EmplaceStep<FCreateSessionHelper>(MoveTemp(CreateSessionHelperParams))
-		.EmplaceLambda([&AccountId, &ExpectedSessionsFound](SubsystemType OnlineSubsystem)
+		.EmplaceLambda([&AccountId, &ExpectedSessionsFound](const IOnlineServicesPtr& OnlineSubsystem)
 			{
 				FGetAllSessions::Params OpGetAllParams;
 				OpGetAllParams.LocalAccountId = AccountId;
@@ -103,7 +103,7 @@ SESSIONS_TEST_CASE("If I call GetSessionByName with an empty session name, I get
 	FAccountId AccountId;
 
 	GetLoginPipeline({ AccountId })
-		.EmplaceLambda([](SubsystemType OnlineSubsystem)
+		.EmplaceLambda([](const IOnlineServicesPtr& OnlineSubsystem)
 			{
 				FGetSessionByName::Params OpParams;
 				OpParams.LocalName = TEXT("");
@@ -122,7 +122,7 @@ SESSIONS_TEST_CASE("If I call GetSessionByName with an unregistered session name
 	FAccountId AccountId;
 
 	GetLoginPipeline({ AccountId })
-		.EmplaceLambda([](SubsystemType OnlineSubsystem)
+		.EmplaceLambda([](const IOnlineServicesPtr& OnlineSubsystem)
 			{
 				FGetSessionByName::Params OpGetByNameParams;
 				OpGetByNameParams.LocalName = TEXT("GetUnregisteredSessionName");
@@ -165,7 +165,7 @@ SESSIONS_TEST_CASE("If I call GetSessionByName with valid information, it return
 	
 	LoginPipeline
 		.EmplaceStep<FCreateSessionHelper>(MoveTemp(CreateSessionHelperParams))
-		.EmplaceLambda([](SubsystemType OnlineSubsystem)
+		.EmplaceLambda([](const IOnlineServicesPtr& OnlineSubsystem)
 			{
 				FGetSessionByName::Params OpGetByNameParams;
 				OpGetByNameParams.LocalName = TEXT("GetSessionByNameValidName");
@@ -186,7 +186,7 @@ SESSIONS_TEST_CASE("If I call GetSessionById with an invalid session id, I get a
 	FAccountId AccountId;
 
 	GetLoginPipeline({ AccountId })
-		.EmplaceLambda([](SubsystemType OnlineSubsystem)
+		.EmplaceLambda([](const IOnlineServicesPtr& OnlineSubsystem)
 			{
 				FGetSessionById::Params OpGetByIdParams;
 				OpGetByIdParams.SessionId = FOnlineSessionId();
@@ -231,7 +231,7 @@ SESSIONS_TEST_CASE("If I call GetSessionById with a valid but unregistered sessi
 
 	LoginPipeline
 		.EmplaceStep<FCreateSessionHelper>(MoveTemp(CreateSessionHelperParams))
-		.EmplaceLambda([&OpGetByIdParams](SubsystemType OnlineSubsystem)
+		.EmplaceLambda([&OpGetByIdParams](const IOnlineServicesPtr& OnlineSubsystem)
 			{
 				ISessionsPtr SessionsInterface = OnlineSubsystem->GetSessionsInterface();
 				TOnlineResult<FGetSessionByName> Result = SessionsInterface->GetSessionByName({ TEXT("GetUnregisteredSessionByIdName") });
@@ -240,7 +240,7 @@ SESSIONS_TEST_CASE("If I call GetSessionById with a valid but unregistered sessi
 				OpGetByIdParams.SessionId = Result.GetOkValue().Session->GetSessionId();
 			})
 		.EmplaceStep<FLeaveSessionHelper>(MoveTemp(LeaveSessionHelperParams))
-		.EmplaceLambda([&OpGetByIdParams](SubsystemType OnlineSubsystem)
+		.EmplaceLambda([&OpGetByIdParams](const IOnlineServicesPtr& OnlineSubsystem)
 			{
 				ISessionsPtr SessionsInterface = OnlineSubsystem->GetSessionsInterface();
 				TOnlineResult<FGetSessionById> Result = SessionsInterface->GetSessionById(MoveTemp(OpGetByIdParams));
@@ -279,7 +279,7 @@ SESSIONS_TEST_CASE("If I call GetSessionById with a valid id for a valid session
 
 	LoginPipeline
 		.EmplaceStep<FCreateSessionHelper>(MoveTemp(CreateSessionHelperParams))
-		.EmplaceLambda([&OpGetByIdParams](SubsystemType OnlineSubsystem)
+		.EmplaceLambda([&OpGetByIdParams](const IOnlineServicesPtr& OnlineSubsystem)
 			{
 				ISessionsPtr SessionsInterface = OnlineSubsystem->GetSessionsInterface();
 				TOnlineResult<FGetSessionByName> Result = SessionsInterface->GetSessionByName({ TEXT("GetSessionByValidIdName") });
@@ -287,7 +287,7 @@ SESSIONS_TEST_CASE("If I call GetSessionById with a valid id for a valid session
 
 				OpGetByIdParams.SessionId = Result.GetOkValue().Session->GetSessionId();
 			})
-		.EmplaceLambda([&OpGetByIdParams](SubsystemType OnlineSubsystem)
+		.EmplaceLambda([&OpGetByIdParams](const IOnlineServicesPtr& OnlineSubsystem)
 			{
 				ISessionsPtr SessionsInterface = OnlineSubsystem->GetSessionsInterface();
 				TOnlineResult<FGetSessionById> Result = SessionsInterface->GetSessionById(MoveTemp(OpGetByIdParams));
@@ -303,7 +303,7 @@ SESSIONS_TEST_CASE("If I call GetSessionById with a valid id for a valid session
 SESSIONS_TEST_CASE("If I call GetPresenceSession with an invalid id, I get an error", EG_SESSIONS_GETSESSIONS_TAG)
 {
 	GetPipeline()
-		.EmplaceLambda([](SubsystemType OnlineSubsystem)
+		.EmplaceLambda([](const IOnlineServicesPtr& OnlineSubsystem)
 			{
 				FGetPresenceSession::Params OpGetPresenceParams;
 				OpGetPresenceParams.LocalAccountId = FAccountId();
@@ -322,7 +322,7 @@ SESSIONS_TEST_CASE("If I call GetPresenceSession with an unregistered id, I get
 	FAccountId AccountId;
 
 	GetLoginPipeline({ AccountId })
-		.EmplaceLambda([&AccountId](SubsystemType OnlineSubsystem)
+		.EmplaceLambda([&AccountId](const IOnlineServicesPtr& OnlineSubsystem)
 			{
 				FGetPresenceSession::Params OpGetPresenceParams;
 				OpGetPresenceParams.LocalAccountId = AccountId;
@@ -363,7 +363,7 @@ SESSIONS_TEST_CASE("If I call GetPresenceSession with a valid id, I get a valid
 
 	LoginPipeline
 		.EmplaceStep<FCreateSessionHelper>(MoveTemp(CreateSessionHelperParams))
-		.EmplaceLambda([&AccountId](SubsystemType OnlineSubsystem)
+		.EmplaceLambda([&AccountId](const IOnlineServicesPtr& OnlineSubsystem)
 			{
 				FGetPresenceSession::Params OpGetPresenceParams;
 				OpGetPresenceParams.LocalAccountId = AccountId;
@@ -383,7 +383,7 @@ SESSIONS_TEST_CASE("If I call GetSessionInviteById with an invalid account id, I
 	OpGetInviteByIdParams.LocalAccountId = FAccountId();
 
 	GetPipeline()
-		.EmplaceLambda([&OpGetInviteByIdParams](SubsystemType OnlineSubsystem)
+		.EmplaceLambda([&OpGetInviteByIdParams](const IOnlineServicesPtr& OnlineSubsystem)
 			{
 				ISessionsPtr SessionsInterface = OnlineSubsystem->GetSessionsInterface();
 				TOnlineResult<FGetSessionInviteById> Result = SessionsInterface->GetSessionInviteById(MoveTemp(OpGetInviteByIdParams));
@@ -408,7 +408,7 @@ SESSIONS_TEST_CASE("If I call GetSessionInviteById with an invalid session invit
 	OpGetInviteByIdParams.LocalAccountId = AccountId;
 
 	LoginPipeline
-		.EmplaceLambda([&OpGetInviteByIdParams](SubsystemType OnlineSubsystem)
+		.EmplaceLambda([&OpGetInviteByIdParams](const IOnlineServicesPtr& OnlineSubsystem)
 			{
 				ISessionsPtr SessionsInterface = OnlineSubsystem->GetSessionsInterface();
 				TOnlineResult<FGetSessionInviteById> Result = SessionsInterface->GetSessionInviteById(MoveTemp(OpGetInviteByIdParams));
@@ -465,7 +465,7 @@ SESSIONS_TEST_CASE("If I call GetSessionInviteById with a valid account id, but
 				OpGetInviteByIdParams.SessionInviteId = InSessionInviteId;
 			})
 		.EmplaceStep<FTickForTime>(FTimespan::FromMilliseconds(1000))
-		.EmplaceLambda([&OpGetInviteByIdParams, &FirstAccountId](SubsystemType OnlineSubsystem)
+		.EmplaceLambda([&OpGetInviteByIdParams, &FirstAccountId](const IOnlineServicesPtr& OnlineSubsystem)
 			{
 				ISessionsPtr SessionsInterface = OnlineSubsystem->GetSessionsInterface();
 				TOnlineResult<FGetSessionInviteById> Result = SessionsInterface->GetSessionInviteById(MoveTemp(OpGetInviteByIdParams));
@@ -531,7 +531,7 @@ SESSIONS_TEST_CASE("If I call GetSessionInviteById with valid invite id, but wit
 			})
 		.EmplaceStep<FTickForTime>(FTimespan::FromMilliseconds(1000))
 		.EmplaceStep<FRejectSessionInviteHelper>(MoveTemp(RejectSessionInviteHelperParams))
-		.EmplaceLambda([&OpGetInviteByIdParams, &FirstAccountId](SubsystemType OnlineSubsystem)
+		.EmplaceLambda([&OpGetInviteByIdParams, &FirstAccountId](const IOnlineServicesPtr& OnlineSubsystem)
 			{
 				ISessionsPtr SessionsInterface = OnlineSubsystem->GetSessionsInterface();
 				TOnlineResult<FGetSessionInviteById> Result = SessionsInterface->GetSessionInviteById(MoveTemp(OpGetInviteByIdParams));
@@ -588,7 +588,7 @@ SESSIONS_TEST_CASE("If I call GetSessionInviteById with a valid data, I get a va
 				OpGetInviteByIdParams.SessionInviteId = InSessionInviteId;
 			})
 		.EmplaceStep<FTickForTime>(FTimespan::FromMilliseconds(1000))
-		.EmplaceLambda([&OpGetInviteByIdParams, &FirstAccountId](SubsystemType OnlineSubsystem)
+		.EmplaceLambda([&OpGetInviteByIdParams, &FirstAccountId](const IOnlineServicesPtr& OnlineSubsystem)
 			{
 				ISessionsPtr SessionsInterface = OnlineSubsystem->GetSessionsInterface();
 				TOnlineResult<FGetSessionInviteById> Result = SessionsInterface->GetSessionInviteById(MoveTemp(OpGetInviteByIdParams));
diff --git a/Engine/Source/Programs/Online/OnlineTestsCore/Source/Private/Tests/Sessions/JoinSessionTests.cpp b/Engine/Source/Programs/Online/OnlineTestsCore/Source/Private/Tests/Sessions/JoinSessionTests.cpp
--- a/Engine/Source/Programs/Online/OnlineTestsCore/Source/Private/Tests/Sessions/JoinSessionTests.cpp
+++ b/Engine/Source/Programs/Online/OnlineTestsCore/Source/Private/Tests/Sessions/JoinSessionTests.cpp
@@ -88,7 +88,7 @@ SESSIONS_TEST_CASE("If I call JoinSession with an empty session name, I get an e
 
 	LoginPipeline
 		.EmplaceStep<FCreateSessionHelper>(MoveTemp(CreateSessionHelperParams))
-		.EmplaceLambda([&JoinSessionHelperParams](SubsystemType OnlineSubsystem)
+		.EmplaceLambda([&JoinSessionHelperParams](const IOnlineServicesPtr& OnlineSubsystem)
 			{
 				ISessionsPtr SessionsInterface = OnlineSubsystem->GetSessionsInterface();
 				TOnlineResult<FGetSessionByName> Result = SessionsInterface->GetSessionByName({ TEXT("EmptySessionNameJoin") });
@@ -186,14 +186,14 @@ SESSIONS_TEST_CASE("If I call JoinSession with an name already in use, I get an
 	FirstLeaveSessionHelperParams.OpParams->LocalAccountId = FirstAccountId;
 	SecondLeaveSessionHelperParams.OpParams->LocalAccountId = FirstAccountId;
 
-	const uint32_t ExpectedSessionsFound = 1;
+	constexpr uint32_t ExpectedSessionsFound = 1;
 
 	LoginPipeline
 		.EmplaceStep<FCreateSessionHelper>(MoveTemp(FirstCreateSessionHelperParams))
 		.EmplaceStep<FCreateSessionHelper>(MoveTemp(SecondCreateSessionHelperParams))
 		.EmplaceStep<FAddSessionMemberHelper>(MoveTemp(FirstAddSessionMemberHelperParams))
 		.EmplaceStep<FAddSessionMemberHelper>(MoveTemp(SecondAddSessionMemberHelperParams))
-		.EmplaceLambda([&FirstFindSessionsHelperParams](SubsystemType OnlineSubsystem)
+		.EmplaceLambda([&FirstFindSessionsHelperParams](const IOnlineServicesPtr& OnlineSubsystem)
 			{
 				ISessionsPtr SessionsInterface = OnlineSubsystem->GetSessionsInterface();
 				TOnlineResult<FGetSessionByName> Result = SessionsInterface->GetSessionByName({ TEXT("SessionNameInUseJoin") });
@@ -201,7 +201,7 @@ SESSIONS_TEST_CASE("If I call JoinSession with an name already in use, I get an
 
 				FirstFindSessionsHelperParams.OpParams->SessionId = Result.GetOkValue().Session->GetSessionId();
 			})
-		.EmplaceLambda([&JoinSessionHelperParams, &SecondFindSessionsHelperParams](SubsystemType OnlineSubsystem)
+		.EmplaceLambda([&JoinSessionHelperParams, &SecondFindSessionsHelperParams](const IOnlineServicesPtr& OnlineSubsystem)
 			{
 				ISessionsPtr SessionsInterface = OnlineSubsystem->GetSessionsInterface();
 				TOnlineResult<FGetSessionByName> Result = SessionsInterface->GetSessionByName({ TEXT("SessionNameInUseJoin1") });
@@ -288,11 +288,11 @@ SESSIONS_TEST_CASE("If I call JoinSession with a valid but unregistered session
 	FirstLeaveSessionHelperParams.OpParams->LocalAccountId = FirstAccountId;
 	SecondLeaveSessionHelperParams.OpParams->LocalAccountId = SecondAccountId;
 
-	const uint32_t ExpectedSessionsFound = 1;
+	constexpr uint32_t ExpectedSessionsFound = 1;
 
 	LoginPipeline
 		.EmplaceStep<FCreateSessionHelper>(MoveTemp(CreateSessionHelperParams))
-		.EmplaceLambda([&FirstJoinSessionHelperParams, &SecondJoinSessionHelperParams, &FirstFindSessionsHelperParams, &SecondFindSessionsHelperParams](SubsystemType OnlineSubsystem)
+		.EmplaceLambda([&FirstJoinSessionHelperParams, &SecondJoinSessionHelperParams, &FirstFindSessionsHelperParams, &SecondFindSessionsHelperParams](const IOnlineServicesPtr& OnlineSubsystem)
 			{
 				ISessionsPtr SessionsInterface = OnlineSubsystem->GetSessionsInterface();
 				TOnlineResult<FGetSessionByName> Result = SessionsInterface->GetSessionByName({ TEXT("SessionIdNotFoundJoinName") });
@@ -368,11 +368,11 @@ SESSIONS_TEST_CASE("If I call JoinSession for a session I'm already registered i
 
 	FindSessionsHelperParams.OpParams->LocalAccountId = SecondAccountId;
 
-	const uint32_t ExpectedSessionsFound = 1;
+	constexpr uint32_t ExpectedSessionsFound = 1;
 
 	LoginPipeline
 		.EmplaceStep<FCreateSessionHelper>(MoveTemp(CreateSessionHelperParams))
-		.EmplaceLambda([&FirstJoinSessionHelperParams, &SecondJoinSessionHelperParams, &FindSessionsHelperParams](SubsystemType OnlineSubsystem)
+		.EmplaceLambda([&FirstJoinSessionHelperParams, &SecondJoinSessionHelperParams, &FindSessionsHelperParams](const IOnlineServicesPtr& OnlineSubsystem)
 			{
 				ISessionsPtr SessionsInterface = OnlineSubsystem->GetSessionsInterface();
 				TOnlineResult<FGetSessionByName> Result = SessionsInterface->GetSessionByName({ TEXT("AlreadyInSessionJoinName") });
@@ -441,11 +441,11 @@ SESSIONS_TEST_CASE("If I call JoinSession for a session that is not joinable, I
 	JoinSessionHelperParams.OpParams->LocalAccountId = SecondAccountId;
 	LeaveSessionHelperParams.OpParams->LocalAccountId = FirstAccountId;
 
-	const uint32_t ExpectedSessionsFound = 1;
+	constexpr uint32_t ExpectedSessionsFound = 1;
 
 	LoginPipeline
 		.EmplaceStep<FCreateSessionHelper>(MoveTemp(CreateSessionHelperParams))
-		.EmplaceLambda([&FindSessionsHelperParams, &JoinSessionHelperParams](SubsystemType OnlineSubsystem)
+		.EmplaceLambda([&FindSessionsHelperParams, &JoinSessionHelperParams](const IOnlineServicesPtr& OnlineSubsystem)
 			{
 				ISessionsPtr SessionsInterface = OnlineSubsystem->GetSessionsInterface();
 				TOnlineResult<FGetSessionByName> Result = SessionsInterface->GetSessionByName({ TEXT("SessionNotJoinableName") });
@@ -500,7 +500,7 @@ SESSIONS_TEST_CASE("If I call JoinSession with presence enabled when there is al
 
 	LoginPipeline
 		.EmplaceStep<FCreateSessionHelper>(MoveTemp(CreateSessionHelperParams))
-		.EmplaceLambda([&JoinSessionHelperParams](SubsystemType OnlineSubsystem)
+		.EmplaceLambda([&JoinSessionHelperParams](const IOnlineServicesPtr& OnlineSubsystem)
 			{
 				ISessionsPtr SessionsInterface = OnlineSubsystem->GetSessionsInterface();
 				TOnlineResult<FGetSessionByName> Result = SessionsInterface->GetSessionByName({ TEXT("PresenceJoinName") });
@@ -552,11 +552,11 @@ SESSIONS_TEST_CASE("If I call JoinSession with valid information, the operation
 	JoinSessionHelperParams.OpParams->LocalAccountId = SecondAccountId;
 	LeaveSessionHelperParams.OpParams->LocalAccountId = FirstAccountId;
 
-	const uint32_t ExpectedSessionsFound = 1;
+	constexpr uint32_t ExpectedSessionsFound = 1;
 
 	LoginPipeline
 		.EmplaceStep<FCreateSessionHelper>(MoveTemp(CreateSessionHelperParams))
-		.EmplaceLambda([&FindSessionsHelperParams, &JoinSessionHelperParams](SubsystemType OnlineSubsystem)
+		.EmplaceLambda([&FindSessionsHelperParams, &JoinSessionHelperParams](const IOnlineServicesPtr& OnlineSubsystem)
 			{
 				ISessionsPtr SessionsInterface = OnlineSubsystem->GetSessionsInterface();
 				TOnlineResult<FGetSessionByName> Result = SessionsInterface->GetSessionByName({ TEXT("SessionValidJoinName") });
diff --git a/Engine/Source/Programs/Online/OnlineTestsCore/Source/Private/Tests/Sessions/PresenceSessionTests.cpp b/Engine/Source/Programs/Online/OnlineTestsCore/Source/Private/Tests/Sessions/PresenceSessionTests.cpp
--- a/Engine/Source/Programs/Online/OnlineTestsCore/Source/Private/Tests/Sessions/PresenceSessionTests.cpp
+++ b/Engine/Source/Programs/Online/OnlineTestsCore/Source/Private/Tests/Sessions/PresenceSessionTests.cpp
@@ -18,7 +18,7 @@ SESSIONS_TEST_CASE("If I call IsPresenceSession with an invalid account id, I ge
 	DestroyCurrentServiceModule();
 
 	GetPipeline()
-		.EmplaceLambda([](SubsystemType OnlineSubsystem)
+		.EmplaceLambda([](const IOnlineServicesPtr& OnlineSubsystem)
 			{
 				FIsPresenceSession::Params OpIsPresenceParams;
 				OpIsPresenceParams.LocalAccountId = FAccountId();
@@ -37,7 +37,7 @@ SESSIONS_TEST_CASE("If I call IsPresenceSession with an invalid session id, I ge
 	FAccountId AccountId;
 
 	GetLoginPipeline({ AccountId })
-		.EmplaceLambda([&AccountId](SubsystemType OnlineSubsystem)
+		.EmplaceLambda([&AccountId](const IOnlineServicesPtr& OnlineSubsystem)
 			{
 				FIsPresenceSession::Params OpIsPresenceParams;
 				OpIsPresenceParams.LocalAccountId = AccountId;
@@ -88,7 +88,7 @@ SESSIONS_TEST_CASE("If I call IsPresenceSession with an unregistered account id,
 
 	LoginPipeline
 		.EmplaceStep<FCreateSessionHelper>(MoveTemp(CreateSessionHelperParams))
-		.EmplaceLambda([&OpIsPresenceParams, &OpGetByNameParams](SubsystemType OnlineSubsystem)
+		.EmplaceLambda([&OpIsPresenceParams, &OpGetByNameParams](const IOnlineServicesPtr& OnlineSubsystem)
 			{			
 				ISessionsPtr SessionsInterface = OnlineSubsystem->GetSessionsInterface();
 				TOnlineResult<FGetSessionByName> Result = SessionsInterface->GetSessionByName(MoveTemp(OpGetByNameParams));
@@ -97,7 +97,7 @@ SESSIONS_TEST_CASE("If I call IsPresenceSession with an unregistered account id,
 				OpIsPresenceParams.SessionId = Result.GetOkValue().Session->GetSessionId();
 			})
 		.EmplaceStep<FLeaveSessionHelper>(MoveTemp(LeaveSessionHelperParams))
-		.EmplaceLambda([&OpIsPresenceParams](SubsystemType OnlineSubsystem)
+		.EmplaceLambda([&OpIsPresenceParams](const IOnlineServicesPtr& OnlineSubsystem)
 			{
 				ISessionsPtr SessionsInterface = OnlineSubsystem->GetSessionsInterface();
 				TOnlineResult<FIsPresenceSession> Result = SessionsInterface->IsPresenceSession(MoveTemp(OpIsPresenceParams));
@@ -161,7 +161,7 @@ SESSIONS_TEST_CASE("If I call IsPresenceSession with valid information, it retur
 		.EmplaceStep<FCreateSessionHelper>(MoveTemp(FirstCreateSessionHelperParams))
 		.EmplaceStep<FCreateSessionHelper>(MoveTemp(SecondCreateSessionHelperParams))
 		.EmplaceStep<FTickForTime>(FTimespan::FromMilliseconds(6000))
-		.EmplaceLambda([&OpFirstIsPresenceParams, &OpSecondIsPresenceParams](SubsystemType OnlineSubsystem)
+		.EmplaceLambda([&OpFirstIsPresenceParams, &OpSecondIsPresenceParams](const IOnlineServicesPtr& OnlineSubsystem)
 			{
 				ISessionsPtr SessionsInterface = OnlineSubsystem->GetSessionsInterface();
 				TOnlineResult<FGetSessionByName> Result = SessionsInterface->GetSessionByName({ TEXT("IsPresenceEnableSessionName") });
@@ -174,7 +174,7 @@ SESSIONS_TEST_CASE("If I call IsPresenceSession with valid information, it retur
 
 				OpSecondIsPresenceParams.SessionId = Result.GetOkValue().Session->GetSessionId();
 			})
-		.EmplaceLambda([&OpFirstIsPresenceParams, &OpSecondIsPresenceParams](SubsystemType OnlineSubsystem)
+		.EmplaceLambda([&OpFirstIsPresenceParams, &OpSecondIsPresenceParams](const IOnlineServicesPtr& OnlineSubsystem)
 			{
 				ISessionsPtr SessionsInterface = OnlineSubsystem->GetSessionsInterface();
 				TOnlineResult<FIsPresenceSession> Result = SessionsInterface->IsPresenceSession(MoveTemp(OpFirstIsPresenceParams));
@@ -194,7 +194,7 @@ SESSIONS_TEST_CASE("If I call IsPresenceSession with valid information, it retur
 SESSIONS_TEST_CASE("If I call SetPresenceSession for EOS, I get an error", EG_SESSIONS_PRESENCESESSIONEOS_TAG)
 {
 	GetPipeline()
-		.EmplaceLambda([](SubsystemType OnlineSubsystem)
+		.EmplaceLambda([](const IOnlineServicesPtr& OnlineSubsystem)
 			{
 				FSetPresenceSession::Params OpSetPresenceParams;
 				
@@ -210,7 +210,7 @@ SESSIONS_TEST_CASE("If I call SetPresenceSession for EOS, I get an error", EG_SE
 SESSIONS_TEST_CASE("If I call SetPresenceSession with an invalid account id, I get an error", EG_SESSIONS_PRESENCESESSIONNULL_TAG)
 {
 	GetPipeline()
-		.EmplaceLambda([](SubsystemType OnlineSubsystem)
+		.EmplaceLambda([](const IOnlineServicesPtr& OnlineSubsystem)
 			{
 				FSetPresenceSession::Params OpSetPresenceParams;
 				OpSetPresenceParams.LocalAccountId = FAccountId();
@@ -228,7 +228,7 @@ SESSIONS_TEST_CASE("If I call SetPresenceSession with an invalid session id, I g
 {
 	FAccountId AccountId;
 	GetLoginPipeline({ AccountId })
-		.EmplaceLambda([&AccountId](SubsystemType OnlineSubsystem)
+		.EmplaceLambda([&AccountId](const IOnlineServicesPtr& OnlineSubsystem)
 			{
 				FSetPresenceSession::Params OpSetPresenceParams;
 				OpSetPresenceParams.LocalAccountId = AccountId;
@@ -271,7 +271,7 @@ SESSIONS_TEST_CASE("If I call SetPresenceSession with valid data, the operation
 
 	LoginPipeline
 		.EmplaceStep<FCreateSessionHelper>(MoveTemp(CreateSessionHelperParams))
-		.EmplaceLambda([&OpSetPresenceParams](SubsystemType OnlineSubsystem)
+		.EmplaceLambda([&OpSetPresenceParams](const IOnlineServicesPtr& OnlineSubsystem)
 			{
 				ISessionsPtr SessionsInterface = OnlineSubsystem->GetSessionsInterface();
 				TOnlineResult<FGetSessionByName> Result = SessionsInterface->GetSessionByName({ TEXT("SetPresenceValidName") });
@@ -279,7 +279,7 @@ SESSIONS_TEST_CASE("If I call SetPresenceSession with valid data, the operation
 
 				OpSetPresenceParams.SessionId = Result.GetOkValue().Session->GetSessionId();
 			})
-		.EmplaceLambda([&OpSetPresenceParams](SubsystemType OnlineSubsystem)
+		.EmplaceLambda([&OpSetPresenceParams](const IOnlineServicesPtr& OnlineSubsystem)
 			{
 				ISessionsPtr SessionsInterface = OnlineSubsystem->GetSessionsInterface();
 				TOnlineResult<FSetPresenceSession> Result = SessionsInterface->SetPresenceSession(MoveTemp(OpSetPresenceParams));
@@ -294,7 +294,7 @@ SESSIONS_TEST_CASE("If I call SetPresenceSession with valid data, the operation
 SESSIONS_TEST_CASE("If I call ClearPresenceSession for EOS, I get an error", EG_SESSIONS_PRESENCESESSIONEOS_TAG)
 {
 	GetPipeline()
-		.EmplaceLambda([](SubsystemType OnlineSubsystem)
+		.EmplaceLambda([](const IOnlineServicesPtr& OnlineSubsystem)
 			{
 				FClearPresenceSession::Params OpClearPresenceParams;
 
@@ -310,7 +310,7 @@ SESSIONS_TEST_CASE("If I call ClearPresenceSession for EOS, I get an error", EG_
 SESSIONS_TEST_CASE("If I call ClearPresenceSession with an invalid account id, I get an error", EG_SESSIONS_PRESENCESESSIONNULL_TAG)
 {
 	GetPipeline()
-		.EmplaceLambda([](SubsystemType OnlineSubsystem)
+		.EmplaceLambda([](const IOnlineServicesPtr& OnlineSubsystem)
 			{
 				FClearPresenceSession::Params OpClearPresenceParams;
 				OpClearPresenceParams.LocalAccountId = FAccountId();
@@ -349,7 +349,7 @@ SESSIONS_TEST_CASE("If I call ClearPresenceSession with valid data, the operatio
 
 	LoginPipeline
 		.EmplaceStep<FCreateSessionHelper>(MoveTemp(CreateSessionHelperParams))
-		.EmplaceLambda([&AccountId](SubsystemType OnlineSubsystem)
+		.EmplaceLambda([&AccountId](const IOnlineServicesPtr& OnlineSubsystem)
 			{
 				FClearPresenceSession::Params OpClearPresenceParams;
 				OpClearPresenceParams.LocalAccountId = AccountId;
diff --git a/Engine/Source/Programs/Online/OnlineTestsCore/Source/Private/Tests/Sessions/SendRejectSessionInvitesTests.cpp b/Engine/Source/Programs/Online/OnlineTestsCore/Source/Private/Tests/Sessions/SendRejectSessionInvitesTests.cpp
--- a/Engine/Source/Programs/Online/OnlineTestsCore/Source/Private/Tests/Sessions/SendRejectSessionInvitesTests.cpp
+++ b/Engine/Source/Programs/Online/OnlineTestsCore/Source/Private/Tests/Sessions/SendRejectSessionInvitesTests.cpp
@@ -190,12 +190,12 @@ SESSIONS_TEST_CASE("If I call SendSessionInvite with valid data, the operation c
 	SendSessionInviteHelperParams.OpParams->TargetUsers.Add(SecondAccountId);
 	LeaveSessionHelperParams.OpParams->LocalAccountId = FirstAccountId;
 
-	const uint32_t ExpectedSessionInvitesNum = 1;
+	constexpr uint32_t ExpectedSessionInvitesNum = 1;
 
 	LoginPipeline
 		.EmplaceStep<FCreateSessionHelper>(MoveTemp(CreateSessionHelperParams))
 		.EmplaceStep<FTickForTime>(FTimespan::FromMilliseconds(1000))
-		.EmplaceLambda([&OpGetAllSessionInvitesParams](SubsystemType OnlineSubsystem)
+		.EmplaceLambda([&OpGetAllSessionInvitesParams](const IOnlineServicesPtr& OnlineSubsystem)
 			{
 				ISessionsPtr SessionsInterface = OnlineSubsystem->GetSessionsInterface();
 				TOnlineResult<FGetAllSessionInvites> Result = SessionsInterface->GetAllSessionInvites(MoveTemp(OpGetAllSessionInvitesParams));
@@ -203,7 +203,7 @@ SESSIONS_TEST_CASE("If I call SendSessionInvite with valid data, the operation c
 			})
 		.EmplaceStep<FSendSessionInviteHelper>(MoveTemp(SendSessionInviteHelperParams))
 		.EmplaceStep<FTickForTime>(FTimespan::FromMilliseconds(1000))
-		.EmplaceLambda([&OpGetAllSessionInvitesParams, &ExpectedSessionInvitesNum](SubsystemType OnlineSubsystem)
+		.EmplaceLambda([&OpGetAllSessionInvitesParams, &ExpectedSessionInvitesNum](const IOnlineServicesPtr& OnlineSubsystem)
 			{
 				ISessionsPtr SessionsInterface = OnlineSubsystem->GetSessionsInterface();
 				TOnlineResult<FGetAllSessionInvites> Result = SessionsInterface->GetAllSessionInvites(MoveTemp(OpGetAllSessionInvitesParams));
@@ -293,7 +293,7 @@ SESSIONS_TEST_CASE("If I call RejectSessionInvite with valid data, the operation
 
 	LeaveSessionHelperParams.OpParams->LocalAccountId = FirstAccountId;
 
-	const uint32_t ExpectedSessionInvitesNum = 1;
+	constexpr uint32_t ExpectedSessionInvitesNum = 1;
 
 	LoginPipeline
 		.EmplaceStep<FCreateSessionHelper>(MoveTemp(CreateSessionHelperParams))
@@ -303,7 +303,7 @@ SESSIONS_TEST_CASE("If I call RejectSessionInvite with valid data, the operation
 				RejectSessionInviteHelperParams.OpParams->SessionInviteId = InSessionInviteId;
 			})
 		.EmplaceStep<FTickForTime>(FTimespan::FromMilliseconds(1000))
-		.EmplaceLambda([&OpGetAllSessionInvitesParams, &ExpectedSessionInvitesNum](SubsystemType OnlineSubsystem)
+		.EmplaceLambda([&OpGetAllSessionInvitesParams, &ExpectedSessionInvitesNum](const IOnlineServicesPtr& OnlineSubsystem)
 			{
 				ISessionsPtr SessionsInterface = OnlineSubsystem->GetSessionsInterface();
 				TOnlineResult<FGetAllSessionInvites> Result = SessionsInterface->GetAllSessionInvites(MoveTemp(OpGetAllSessionInvitesParams));
@@ -311,7 +311,7 @@ SESSIONS_TEST_CASE("If I call RejectSessionInvite with valid data, the operation
 			})
 		.EmplaceStep<FRejectSessionInviteHelper>(MoveTemp(RejectSessionInviteHelperParams))
 		.EmplaceStep<FTickForTime>(FTimespan::FromMilliseconds(1000))
-		.EmplaceLambda([&OpGetAllSessionInvitesParams](SubsystemType OnlineSubsystem)
+		.EmplaceLambda([&OpGetAllSessionInvitesParams](const IOnlineServicesPtr& OnlineSubsystem)
 			{
 				ISessionsPtr SessionsInterface = OnlineSubsystem->GetSessionsInterface();
 				TOnlineResult<FGetAllSessionInvites> Result = SessionsInterface->GetAllSessionInvites(MoveTemp(OpGetAllSessionInvitesParams));
diff --git a/Engine/Source/Programs/Online/OnlineTestsCore/Source/Private/Tests/Sessions/UpdateSessionSettingsTests.cpp b/Engine/Source/Programs/Online/OnlineTestsCore/Source/Private/Tests/Sessions/UpdateSessionSettingsTests.cpp
--- a/Engine/Source/Programs/Online/OnlineTestsCore/Source/Private/Tests/Sessions/UpdateSessionSettingsTests.cpp
+++ b/Engine/Source/Programs/Online/OnlineTestsCore/Source/Private/Tests/Sessions/UpdateSessionSettingsTests.cpp
@@ -173,7 +173,7 @@ UPDATESESSIONS_TEST_CASE("If I call UpdateSessionSettings with valid data, the o
 	LoginPipeline
 		.EmplaceStep<FCreateSessionHelper>(MoveTemp(CreateSessionHelperParams))
 		.EmplaceStep<FUpdateSessionSettingsHelper>(MoveTemp(UpdateSessionSettingsHelperParams))
-		.EmplaceLambda([&OpGetByNameParams, &UpdateSessionSettingsHelperParams](SubsystemType OnlineSubsystem)
+		.EmplaceLambda([&OpGetByNameParams, &UpdateSessionSettingsHelperParams](const IOnlineServicesPtr& OnlineSubsystem)
 			{
 				ISessionsPtr SessionsInterface = OnlineSubsystem->GetSessionsInterface();
 				TOnlineResult<FGetSessionByName> Result = SessionsInterface->GetSessionByName(MoveTemp(OpGetByNameParams));
diff --git a/Engine/Source/Programs/Online/OnlineTestsCore/Source/Private/Tests/Social/AcceptFriendInviteTests.cpp b/Engine/Source/Programs/Online/OnlineTestsCore/Source/Private/Tests/Social/AcceptFriendInviteTests.cpp
--- a/Engine/Source/Programs/Online/OnlineTestsCore/Source/Private/Tests/Social/AcceptFriendInviteTests.cpp
+++ b/Engine/Source/Programs/Online/OnlineTestsCore/Source/Private/Tests/Social/AcceptFriendInviteTests.cpp
@@ -65,22 +65,22 @@ SOCIAL_TEST_CASE("Verify that AcceptFriendInvite returns a fail message if the l
 	AcceptInviteHelperParams.OpParams->TargetAccountId = SecondAccountId;
 
 	LoginPipeline
-		.EmplaceLambda([&FirstAccountId, &SecondAccountId, FirstAccountPlatformUserId, SecondAccountPlatformUserId](SubsystemType OnlineSubsystem)
+		.EmplaceLambda([&FirstAccountId, &SecondAccountId, FirstAccountPlatformUserId, SecondAccountPlatformUserId](const IOnlineServicesPtr& OnlineSubsystem)
 			{
 				UE::Online::IAuthPtr OnlineAuthPtr = OnlineSubsystem->GetAuthInterface();
 				REQUIRE(OnlineAuthPtr);
 
-				UE::Online::TOnlineResult<UE::Online::FAuthGetLocalOnlineUserByOnlineAccountId> FirstUserPlatfromUserIdResult = OnlineAuthPtr->GetLocalOnlineUserByOnlineAccountId({ FirstAccountId });
-				UE::Online::TOnlineResult<UE::Online::FAuthGetLocalOnlineUserByOnlineAccountId> SecondUserPlatfromUserIdResult = OnlineAuthPtr->GetLocalOnlineUserByOnlineAccountId({ SecondAccountId });
+				UE::Online::TOnlineResult<UE::Online::FAuthGetLocalOnlineUserByOnlineAccountId> FirstUserPlatformUserIdResult = OnlineAuthPtr->GetLocalOnlineUserByOnlineAccountId({ FirstAccountId });
+				UE::Online::TOnlineResult<UE::Online::FAuthGetLocalOnlineUserByOnlineAccountId> SecondUserPlatformUserIdResult = OnlineAuthPtr->GetLocalOnlineUserByOnlineAccountId({ SecondAccountId });
 
-				REQUIRE(FirstUserPlatfromUserIdResult.IsOk());
-				REQUIRE(SecondUserPlatfromUserIdResult.IsOk());
+				REQUIRE(FirstUserPlatformUserIdResult.IsOk());
+				REQUIRE(SecondUserPlatformUserIdResult.IsOk());
 
-				CHECK(FirstUserPlatfromUserIdResult.TryGetOkValue() != nullptr);
-				CHECK(SecondUserPlatfromUserIdResult.TryGetOkValue() != nullptr);
+				CHECK(FirstUserPlatformUserIdResult.TryGetOkValue() != nullptr);
+				CHECK(SecondUserPlatformUserIdResult.TryGetOkValue() != nullptr);
 
-				*FirstAccountPlatformUserId = FirstUserPlatfromUserIdResult.TryGetOkValue()->AccountInfo->PlatformUserId;
-				*SecondAccountPlatformUserId = SecondUserPlatfromUserIdResult.TryGetOkValue()->AccountInfo->PlatformUserId;
+				*FirstAccountPlatformUserId = FirstUserPlatformUserIdResult.TryGetOkValue()->AccountInfo->PlatformUserId;
+				*SecondAccountPlatformUserId = SecondUserPlatformUserIdResult.TryGetOkValue()->AccountInfo->PlatformUserId;
 			})
 		.EmplaceStep<FAuthLogoutStep>(MoveTemp(FirstAccountPlatformUserId))
 		.EmplaceStep<FAcceptFriendInviteHelper>(MoveTemp(AcceptInviteHelperParams))
diff --git a/Engine/Source/Programs/Online/OnlineTestsCore/Source/Private/Tests/Social/BlockUserTests.cpp b/Engine/Source/Programs/Online/OnlineTestsCore/Source/Private/Tests/Social/BlockUserTests.cpp
--- a/Engine/Source/Programs/Online/OnlineTestsCore/Source/Private/Tests/Social/BlockUserTests.cpp
+++ b/Engine/Source/Programs/Online/OnlineTestsCore/Source/Private/Tests/Social/BlockUserTests.cpp
@@ -65,22 +65,22 @@ SOCIAL_TEST_CASE("Verify that BlockUser returns a fail message if the local user
 	BlockUserHelperHelperParams.OpParams->TargetAccountId = SecondAccountId;
 
 	LoginPipeline
-		.EmplaceLambda([&FirstAccountId, &SecondAccountId, FirstAccountPlatformUserId, SecondAccountPlatformUserId](SubsystemType OnlineSubsystem)
+		.EmplaceLambda([&FirstAccountId, &SecondAccountId, FirstAccountPlatformUserId, SecondAccountPlatformUserId](const IOnlineServicesPtr& OnlineSubsystem)
 			{
 				UE::Online::IAuthPtr OnlineAuthPtr = OnlineSubsystem->GetAuthInterface();
 				REQUIRE(OnlineAuthPtr);
 
-				UE::Online::TOnlineResult<UE::Online::FAuthGetLocalOnlineUserByOnlineAccountId> FirstUserPlatfromUserIdResult = OnlineAuthPtr->GetLocalOnlineUserByOnlineAccountId({ FirstAccountId });
-				UE::Online::TOnlineResult<UE::Online::FAuthGetLocalOnlineUserByOnlineAccountId> SecondUserPlatfromUserIdResult = OnlineAuthPtr->GetLocalOnlineUserByOnlineAccountId({ SecondAccountId });
+				UE::Online::TOnlineResult<UE::Online::FAuthGetLocalOnlineUserByOnlineAccountId> FirstUserPlatformUserIdResult = OnlineAuthPtr->GetLocalOnlineUserByOnlineAccountId({ FirstAccountId });
+				UE::Online::TOnlineResult<UE::Online::FAuthGetLocalOnlineUserByOnlineAccountId> SecondUserPlatformUserIdResult = OnlineAuthPtr->GetLocalOnlineUserByOnlineAccountId({ SecondAccountId });
 
-				REQUIRE(FirstUserPlatfromUserIdResult.IsOk());
-				REQUIRE(SecondUserPlatfromUserIdResult.IsOk());
+				REQUIRE(FirstUserPlatformUserIdResult.IsOk());
+				REQUIRE(SecondUserPlatformUserIdResult.IsOk());
 
-				CHECK(FirstUserPlatfromUserIdResult.TryGetOkValue() != nullptr);
-				CHECK(SecondUserPlatfromUserIdResult.TryGetOkValue() != nullptr);
+				CHECK(FirstUserPlatformUserIdResult.TryGetOkValue() != nullptr);
+				CHECK(SecondUserPlatformUserIdResult.TryGetOkValue() != nullptr);
 
-				*FirstAccountPlatformUserId = FirstUserPlatfromUserIdResult.TryGetOkValue()->AccountInfo->PlatformUserId;
-				*SecondAccountPlatformUserId = SecondUserPlatfromUserIdResult.TryGetOkValue()->AccountInfo->PlatformUserId;
+				*FirstAccountPlatformUserId = FirstUserPlatformUserIdResult.TryGetOkValue()->AccountInfo->PlatformUserId;
+				*SecondAccountPlatformUserId = SecondUserPlatformUserIdResult.TryGetOkValue()->AccountInfo->PlatformUserId;
 			})
 		.EmplaceStep<FAuthLogoutStep>(MoveTemp(FirstAccountPlatformUserId))
 		.EmplaceStep<FBlockUserHelper>(MoveTemp(BlockUserHelperHelperParams))
diff --git a/Engine/Source/Programs/Online/OnlineTestsCore/Source/Private/Tests/Social/GetBlockedUsersTests.cpp b/Engine/Source/Programs/Online/OnlineTestsCore/Source/Private/Tests/Social/GetBlockedUsersTests.cpp
--- a/Engine/Source/Programs/Online/OnlineTestsCore/Source/Private/Tests/Social/GetBlockedUsersTests.cpp
+++ b/Engine/Source/Programs/Online/OnlineTestsCore/Source/Private/Tests/Social/GetBlockedUsersTests.cpp
@@ -15,7 +15,7 @@ SOCIAL_TEST_CASE("Verify that GetBlockedUsers returns a error if call with an in
 	GetBlockedUsersHelperParams.OpParams = &OpGetBlockedUsersParams;
 	GetBlockedUsersHelperParams.OpParams->LocalAccountId = FAccountId();
 
-	SubsystemType OnlineSubsystem = GetSubsystem();
+	IOnlineServicesPtr OnlineSubsystem = GetSubsystem();
 	EOnlineServices ServicesProvider = OnlineSubsystem->GetServicesProvider();
 
 	if (ServicesProvider == EOnlineServices::Epic)
@@ -42,7 +42,7 @@ SOCIAL_TEST_CASE("Verify that GetBlockedUsers returns a fail message if there ar
 	FGetBlockedUsersHelper::FHelperParams GetBlockedUsersHelperParams;
 	GetBlockedUsersHelperParams.OpParams = &OpGetBlockedUsersParams;
 
-	SubsystemType OnlineSubsystem = GetSubsystem();
+	IOnlineServicesPtr OnlineSubsystem = GetSubsystem();
 	EOnlineServices ServicesProvider = OnlineSubsystem->GetServicesProvider();
 
 	if (ServicesProvider == EOnlineServices::Epic)
diff --git a/Engine/Source/Programs/Online/OnlineTestsCore/Source/Private/Tests/Social/QueryBlockedUsersTests.cpp b/Engine/Source/Programs/Online/OnlineTestsCore/Source/Private/Tests/Social/QueryBlockedUsersTests.cpp
--- a/Engine/Source/Programs/Online/OnlineTestsCore/Source/Private/Tests/Social/QueryBlockedUsersTests.cpp
+++ b/Engine/Source/Programs/Online/OnlineTestsCore/Source/Private/Tests/Social/QueryBlockedUsersTests.cpp
@@ -15,7 +15,7 @@ SOCIAL_TEST_CASE("Verify that QueryBlockedUsers returns a error if call with an
 	QueryBlockedUsersParams.OpParams = &OpQueryParams;
 	QueryBlockedUsersParams.OpParams->LocalAccountId = FAccountId();
 	
-	SubsystemType OnlineSubsystem = GetSubsystem();
+	IOnlineServicesPtr OnlineSubsystem = GetSubsystem();
 	EOnlineServices ServicesProvider = OnlineSubsystem->GetServicesProvider();
 	
 	if (ServicesProvider == EOnlineServices::Epic)
diff --git a/Engine/Source/Programs/Online/OnlineTestsCore/Source/Private/Tests/Social/QueryFriendsTests.cpp b/Engine/Source/Programs/Online/OnlineTestsCore/Source/Private/Tests/Social/QueryFriendsTests.cpp
--- a/Engine/Source/Programs/Online/OnlineTestsCore/Source/Private/Tests/Social/QueryFriendsTests.cpp
+++ b/Engine/Source/Programs/Online/OnlineTestsCore/Source/Private/Tests/Social/QueryFriendsTests.cpp
@@ -29,17 +29,17 @@ SOCIAL_TEST_CASE("Verify that QueryFriends returns a fail message if the local u
 
 
 	LoginPipeline
-		.EmplaceLambda([&AccountId, AccountPlatformUserId](SubsystemType OnlineSubsystem)
+		.EmplaceLambda([&AccountId, AccountPlatformUserId](const IOnlineServicesPtr& OnlineSubsystem)
 			{
 				UE::Online::IAuthPtr OnlineAuthPtr = OnlineSubsystem->GetAuthInterface();
 				REQUIRE(OnlineAuthPtr);
 
-				UE::Online::TOnlineResult<UE::Online::FAuthGetLocalOnlineUserByOnlineAccountId> UserPlatfromUserIdResult = OnlineAuthPtr->GetLocalOnlineUserByOnlineAccountId({ AccountId });
+				UE::Online::TOnlineResult<UE::Online::FAuthGetLocalOnlineUserByOnlineAccountId> UserPlatformUserIdResult = OnlineAuthPtr->GetLocalOnlineUserByOnlineAccountId({ AccountId });
 				
-				REQUIRE(UserPlatfromUserIdResult.IsOk());
-				CHECK(UserPlatfromUserIdResult.TryGetOkValue() != nullptr);
+				REQUIRE(UserPlatformUserIdResult.IsOk());
+				CHECK(UserPlatformUserIdResult.TryGetOkValue() != nullptr);
 
-				*AccountPlatformUserId = UserPlatfromUserIdResult.TryGetOkValue()->AccountInfo->PlatformUserId;
+				*AccountPlatformUserId = UserPlatformUserIdResult.TryGetOkValue()->AccountInfo->PlatformUserId;
 			})
 		.EmplaceStep<FAuthLogoutStep>(MoveTemp(AccountPlatformUserId))
 		.EmplaceStep<FQueryFriendsHelper>(MoveTemp(QueryFriendsHelperParams));
diff --git a/Engine/Source/Programs/Online/OnlineTestsCore/Source/Private/Tests/Social/RejectFriendInviteTests.cpp b/Engine/Source/Programs/Online/OnlineTestsCore/Source/Private/Tests/Social/RejectFriendInviteTests.cpp
--- a/Engine/Source/Programs/Online/OnlineTestsCore/Source/Private/Tests/Social/RejectFriendInviteTests.cpp
+++ b/Engine/Source/Programs/Online/OnlineTestsCore/Source/Private/Tests/Social/RejectFriendInviteTests.cpp
@@ -63,22 +63,22 @@ SOCIAL_TEST_CASE("Verify that RejectFriendInvite returns a fail message if the l
 	RejectFriendInviteHelperParams.OpParams->TargetAccountId = SecondAccountId;
 
 	LoginPipeline
-		.EmplaceLambda([&FirstAccountId, &SecondAccountId, FirstAccountPlatformUserId, SecondAccountPlatformUserId](SubsystemType OnlineSubsystem)
+		.EmplaceLambda([&FirstAccountId, &SecondAccountId, FirstAccountPlatformUserId, SecondAccountPlatformUserId](const IOnlineServicesPtr& OnlineSubsystem)
 			{
 				UE::Online::IAuthPtr OnlineAuthPtr = OnlineSubsystem->GetAuthInterface();
 				REQUIRE(OnlineAuthPtr);
 
-				UE::Online::TOnlineResult<UE::Online::FAuthGetLocalOnlineUserByOnlineAccountId> FirstUserPlatfromUserIdResult = OnlineAuthPtr->GetLocalOnlineUserByOnlineAccountId({ FirstAccountId });
-				UE::Online::TOnlineResult<UE::Online::FAuthGetLocalOnlineUserByOnlineAccountId> SecondUserPlatfromUserIdResult = OnlineAuthPtr->GetLocalOnlineUserByOnlineAccountId({ SecondAccountId });
+				UE::Online::TOnlineResult<UE::Online::FAuthGetLocalOnlineUserByOnlineAccountId> FirstUserPlatformUserIdResult = OnlineAuthPtr->GetLocalOnlineUserByOnlineAccountId({ FirstAccountId });
+				UE::Online::TOnlineResult<UE::Online::FAuthGetLocalOnlineUserByOnlineAccountId> SecondUserPlatformUserIdResult = OnlineAuthPtr->GetLocalOnlineUserByOnlineAccountId({ SecondAccountId });
 
-				REQUIRE(FirstUserPlatfromUserIdResult.IsOk());
-				REQUIRE(SecondUserPlatfromUserIdResult.IsOk());
+				REQUIRE(FirstUserPlatformUserIdResult.IsOk());
+				REQUIRE(SecondUserPlatformUserIdResult.IsOk());
 
-				CHECK(FirstUserPlatfromUserIdResult.TryGetOkValue() != nullptr);
-				CHECK(SecondUserPlatfromUserIdResult.TryGetOkValue() != nullptr);
+				CHECK(FirstUserPlatformUserIdResult.TryGetOkValue() != nullptr);
+				CHECK(SecondUserPlatformUserIdResult.TryGetOkValue() != nullptr);
 
-				*FirstAccountPlatformUserId = FirstUserPlatfromUserIdResult.TryGetOkValue()->AccountInfo->PlatformUserId;
-				*SecondAccountPlatformUserId = SecondUserPlatfromUserIdResult.TryGetOkValue()->AccountInfo->PlatformUserId;
+				*FirstAccountPlatformUserId = FirstUserPlatformUserIdResult.TryGetOkValue()->AccountInfo->PlatformUserId;
+				*SecondAccountPlatformUserId = SecondUserPlatformUserIdResult.TryGetOkValue()->AccountInfo->PlatformUserId;
 			})
 		.EmplaceStep<FAuthLogoutStep>(MoveTemp(FirstAccountPlatformUserId))
 		.EmplaceStep<FRejectFriendInviteHelper>(MoveTemp(RejectFriendInviteHelperParams))
diff --git a/Engine/Source/Programs/Online/OnlineTestsCore/Source/Private/Tests/Social/SendFriendInviteTests.cpp b/Engine/Source/Programs/Online/OnlineTestsCore/Source/Private/Tests/Social/SendFriendInviteTests.cpp
--- a/Engine/Source/Programs/Online/OnlineTestsCore/Source/Private/Tests/Social/SendFriendInviteTests.cpp
+++ b/Engine/Source/Programs/Online/OnlineTestsCore/Source/Private/Tests/Social/SendFriendInviteTests.cpp
@@ -67,22 +67,22 @@ SOCIAL_TEST_CASE("Verify that SendFriendInvite returns a fail message if the loc
 	SendFriendInviteHelperParams.OpParams->TargetAccountId = SecondAccountId;
 
 	LoginPipeline
-		.EmplaceLambda([&FirstAccountId, &SecondAccountId, FirstAccountPlatformUserId, SecondAccountPlatformUserId](SubsystemType OnlineSubsystem)
+		.EmplaceLambda([&FirstAccountId, &SecondAccountId, FirstAccountPlatformUserId, SecondAccountPlatformUserId](const IOnlineServicesPtr& OnlineSubsystem)
 			{
 				UE::Online::IAuthPtr OnlineAuthPtr = OnlineSubsystem->GetAuthInterface();
 				REQUIRE(OnlineAuthPtr);
 
-				UE::Online::TOnlineResult<UE::Online::FAuthGetLocalOnlineUserByOnlineAccountId> FirstUserPlatfromUserIdResult = OnlineAuthPtr->GetLocalOnlineUserByOnlineAccountId({ FirstAccountId });
-				UE::Online::TOnlineResult<UE::Online::FAuthGetLocalOnlineUserByOnlineAccountId> SecondUserPlatfromUserIdResult = OnlineAuthPtr->GetLocalOnlineUserByOnlineAccountId({ SecondAccountId });
+				UE::Online::TOnlineResult<UE::Online::FAuthGetLocalOnlineUserByOnlineAccountId> FirstUserPlatformUserIdResult = OnlineAuthPtr->GetLocalOnlineUserByOnlineAccountId({ FirstAccountId });
+				UE::Online::TOnlineResult<UE::Online::FAuthGetLocalOnlineUserByOnlineAccountId> SecondUserPlatformUserIdResult = OnlineAuthPtr->GetLocalOnlineUserByOnlineAccountId({ SecondAccountId });
 
-				REQUIRE(FirstUserPlatfromUserIdResult.IsOk());
-				REQUIRE(SecondUserPlatfromUserIdResult.IsOk());
+				REQUIRE(FirstUserPlatformUserIdResult.IsOk());
+				REQUIRE(SecondUserPlatformUserIdResult.IsOk());
 
-				CHECK(FirstUserPlatfromUserIdResult.TryGetOkValue() != nullptr);
-				CHECK(SecondUserPlatfromUserIdResult.TryGetOkValue() != nullptr);
+				CHECK(FirstUserPlatformUserIdResult.TryGetOkValue() != nullptr);
+				CHECK(SecondUserPlatformUserIdResult.TryGetOkValue() != nullptr);
 
-				*FirstAccountPlatformUserId = FirstUserPlatfromUserIdResult.TryGetOkValue()->AccountInfo->PlatformUserId;
-				*SecondAccountPlatformUserId = SecondUserPlatfromUserIdResult.TryGetOkValue()->AccountInfo->PlatformUserId;
+				*FirstAccountPlatformUserId = FirstUserPlatformUserIdResult.TryGetOkValue()->AccountInfo->PlatformUserId;
+				*SecondAccountPlatformUserId = SecondUserPlatformUserIdResult.TryGetOkValue()->AccountInfo->PlatformUserId;
 			})
 		.EmplaceStep<FAuthLogoutStep>(MoveTemp(FirstAccountPlatformUserId))
 		.EmplaceStep<FSendFriendInviteHelper>(MoveTemp(SendFriendInviteHelperParams))
diff --git a/Engine/Source/Programs/Online/OnlineTestsCore/Source/Private/Tests/Stats/StatsTests.cpp b/Engine/Source/Programs/Online/OnlineTestsCore/Source/Private/Tests/Stats/StatsTests.cpp
--- a/Engine/Source/Programs/Online/OnlineTestsCore/Source/Private/Tests/Stats/StatsTests.cpp
+++ b/Engine/Source/Programs/Online/OnlineTestsCore/Source/Private/Tests/Stats/StatsTests.cpp
@@ -11,7 +11,7 @@
 FTestPipeline&& StatsFixture_1User_With_InitialRandomStats(const OnlineTestBase& InOnlineTestBase, FAccountId& AccountId)
 {
 	return InOnlineTestBase.GetLoginPipeline({ AccountId })
-		.EmplaceAsyncLambda([&AccountId](FAsyncLambdaResult Promise, SubsystemType Services) {
+		.EmplaceAsyncLambda([&AccountId](FAsyncLambdaResult Promise, const IOnlineServicesPtr& Services) {
 			int64 InitValue = (int64)FMath::RandRange(1, 100);
 			TArray<FUserStats> UpdateUsersStats{
 				{ AccountId, { { "Stat_Use_Largest", FStatValue(InitValue)}, {"Stat_Use_Smallest", FStatValue(InitValue)}, { "Stat_Use_Sum", FStatValue(InitValue)}, { "Stat_Use_Set", FStatValue(InitValue)}}},
@@ -20,7 +20,7 @@ FTestPipeline&& StatsFixture_1User_With_InitialRandomStats(const OnlineTestBase&
 		});
 }
 
-void QueryStats_Fixture(FAsyncLambdaResult Promise, SubsystemType Services, const FAccountId& LocalAccountId, const FAccountId& TargetAccountId, const TArray<FString>& StatNames, TMap<FString, FStatValue>& OutStats)
+void QueryStats_Fixture(FAsyncLambdaResult Promise, const IOnlineServicesPtr& Services, const FAccountId& LocalAccountId, const FAccountId& TargetAccountId, const TArray<FString>& StatNames, TMap<FString, FStatValue>& OutStats)
 {
 	Services->GetStatsInterface()->QueryStats({ LocalAccountId, TargetAccountId, StatNames }).OnComplete([Promise=MoveTemp(Promise), &OutStats](const TOnlineResult<FQueryStats>& Result) mutable
 	{
@@ -41,10 +41,10 @@ STATS_TEST_CASE("UpdateStats modify method work as expected", "[tmpExcludedStats
 	int64 Expected_Stat_Use_Set = (int64)FMath::RandRange(1, 100);
 
 	StatsFixture_1User_With_InitialRandomStats(*this, AccountId)
-		.EmplaceAsyncLambda([&AccountId, &InitialStats](FAsyncLambdaResult Promise, SubsystemType Services) {
+		.EmplaceAsyncLambda([&AccountId, &InitialStats](FAsyncLambdaResult Promise, const IOnlineServicesPtr& Services) {
 			QueryStats_Fixture(MoveTemp(Promise), Services, AccountId, AccountId, { "Stat_Use_Largest", "Stat_Use_Smallest", "Stat_Use_Sum", "Stat_Use_Set" }, InitialStats); 
 		})
-		.EmplaceAsyncLambda([&AccountId, &InitialStats, &Expected_Stat_Use_Set](FAsyncLambdaResult Promise, SubsystemType Services) {
+		.EmplaceAsyncLambda([&AccountId, &InitialStats, &Expected_Stat_Use_Set](FAsyncLambdaResult Promise, const IOnlineServicesPtr& Services) {
 			TArray<FUserStats> UpdateUsersStats{ { AccountId, {
 				{"Stat_Use_Largest", FStatValue(InitialStats["Stat_Use_Largest"].GetInt64() - 1)},
 				{"Stat_Use_Smallest", FStatValue(InitialStats["Stat_Use_Largest"].GetInt64() + 1)},
@@ -53,26 +53,26 @@ STATS_TEST_CASE("UpdateStats modify method work as expected", "[tmpExcludedStats
 			}}, };
 			UpdateStats_Fixture(MoveTemp(Promise), Services, AccountId, UpdateUsersStats);
 		})
-		.EmplaceAsyncLambda([&AccountId, &UpdatedStats](FAsyncLambdaResult Promise, SubsystemType Services) {
+		.EmplaceAsyncLambda([&AccountId, &UpdatedStats](FAsyncLambdaResult Promise, const IOnlineServicesPtr& Services) {
 			QueryStats_Fixture(MoveTemp(Promise), Services, AccountId, AccountId, { "Stat_Use_Largest", "Stat_Use_Smallest", "Stat_Use_Sum", "Stat_Use_Set" }, UpdatedStats);
 		})
-		.EmplaceLambda([&AccountId, &InitialStats, &UpdatedStats, &Expected_Stat_Use_Set](SubsystemType Services) {
+		.EmplaceLambda([&AccountId, &InitialStats, &UpdatedStats, &Expected_Stat_Use_Set](const IOnlineServicesPtr& Services) {
 			CHECK_EQUAL(UpdatedStats["Stat_Use_Largest"].GetInt64(), InitialStats["Stat_Use_Largest"].GetInt64()); // Setting to smaller value didn't take effect
 			CHECK_EQUAL(UpdatedStats["Stat_Use_Smallest"].GetInt64(), InitialStats["Stat_Use_Smallest"].GetInt64());  // Setting to bigger value didn't take effect
 			CHECK_EQUAL(UpdatedStats["Stat_Use_Sum"].GetInt64(), InitialStats["Stat_Use_Sum"].GetInt64() + 3);
 			CHECK_EQUAL(UpdatedStats["Stat_Use_Set"].GetInt64(), Expected_Stat_Use_Set);
 		})
-		.EmplaceAsyncLambda([&AccountId, &InitialStats](FAsyncLambdaResult Promise, SubsystemType Services) {
+		.EmplaceAsyncLambda([&AccountId, &InitialStats](FAsyncLambdaResult Promise, const IOnlineServicesPtr& Services) {
 			TArray<FUserStats> UpdateUsersStats{ { AccountId, {
 				{"Stat_Use_Largest", FStatValue(InitialStats["Stat_Use_Largest"].GetInt64() + 1)},
 				{"Stat_Use_Smallest", FStatValue(InitialStats["Stat_Use_Smallest"].GetInt64() - 1)},
 			}}, };
 			UpdateStats_Fixture(MoveTemp(Promise), Services, AccountId, UpdateUsersStats);
 		})
-		.EmplaceAsyncLambda([&AccountId, &UpdatedStats](FAsyncLambdaResult Promise, SubsystemType Services) {
+		.EmplaceAsyncLambda([&AccountId, &UpdatedStats](FAsyncLambdaResult Promise, const IOnlineServicesPtr& Services) {
 			QueryStats_Fixture(MoveTemp(Promise), Services, AccountId, AccountId, { "Stat_Use_Largest", "Stat_Use_Smallest" }, UpdatedStats);
 		})
-		.EmplaceLambda([&InitialStats, &UpdatedStats](SubsystemType Services) {
+		.EmplaceLambda([&InitialStats, &UpdatedStats](const IOnlineServicesPtr& Services) {
 			CHECK_EQUAL(UpdatedStats["Stat_Use_Largest"].GetInt64(), InitialStats["Stat_Use_Largest"].GetInt64() + 1); // Setting to bigger value took effect
 			CHECK_EQUAL(UpdatedStats["Stat_Use_Smallest"].GetInt64(), InitialStats["Stat_Use_Smallest"].GetInt64() - 1);  // Setting to smaller value took effect
 		});
@@ -87,23 +87,23 @@ STATS_TEST_CASE("UpdateStats modify different types work as expected", "[tmpExcl
 	TMap<FString, FStatValue> UpdatedStats;
 
 	GetLoginPipeline({ AccountId })
-		.EmplaceAsyncLambda([&AccountId](FAsyncLambdaResult Promise, SubsystemType Services) {
+		.EmplaceAsyncLambda([&AccountId](FAsyncLambdaResult Promise, const IOnlineServicesPtr& Services) {
 			UpdateStats_Fixture(MoveTemp(Promise), Services, AccountId, { { AccountId, { { "Stat_Type_Bool", FStatValue(true)}, {"Stat_Type_Double", FStatValue(9999.999)}}}, });
 		})
-		.EmplaceAsyncLambda([&AccountId, &InitialStats](FAsyncLambdaResult Promise, SubsystemType Services) {
+		.EmplaceAsyncLambda([&AccountId, &InitialStats](FAsyncLambdaResult Promise, const IOnlineServicesPtr& Services) {
 			QueryStats_Fixture(MoveTemp(Promise), Services, AccountId, AccountId, { "Stat_Type_Bool", "Stat_Type_Double" }, InitialStats); 
 		})
-		.EmplaceAsyncLambda([&AccountId, &InitialStats](FAsyncLambdaResult Promise, SubsystemType Services) {
+		.EmplaceAsyncLambda([&AccountId, &InitialStats](FAsyncLambdaResult Promise, const IOnlineServicesPtr& Services) {
 			TArray<FUserStats> UpdateUsersStats{ { AccountId, {
 				{"Stat_Type_Bool", FStatValue(false)},
 				{"Stat_Type_Double", FStatValue(InitialStats["Stat_Type_Double"].GetDouble() - 0.001)},
 			}}, };
 			UpdateStats_Fixture(MoveTemp(Promise), Services, AccountId, UpdateUsersStats);
 		})
-		.EmplaceAsyncLambda([&AccountId, &UpdatedStats](FAsyncLambdaResult Promise, SubsystemType Services) {
+		.EmplaceAsyncLambda([&AccountId, &UpdatedStats](FAsyncLambdaResult Promise, const IOnlineServicesPtr& Services) {
 			QueryStats_Fixture(MoveTemp(Promise), Services, AccountId, AccountId, { "Stat_Type_Bool", "Stat_Type_Double" }, UpdatedStats);
 		})
-		.EmplaceLambda([&InitialStats, &UpdatedStats](SubsystemType Services) {
+		.EmplaceLambda([&InitialStats, &UpdatedStats](const IOnlineServicesPtr& Services) {
 			CHECK_EQUAL(UpdatedStats["Stat_Type_Bool"].GetBoolean(), false);
 			CHECK(FMath::IsNearlyEqual(UpdatedStats["Stat_Type_Double"].GetDouble(), InitialStats["Stat_Type_Double"].GetDouble() - 0.001, UE_DOUBLE_SMALL_NUMBER));
 		});
@@ -118,16 +118,16 @@ STATS_TEST_CASE("ResetStats works or not implemented", "[tmpExcludedStats]")
 	bool ImplementedStatsReset = false;
 
 	StatsFixture_1User_With_InitialRandomStats(*this, AccountId)
-		.EmplaceAsyncLambda([&AccountId, &UserStats](FAsyncLambdaResult Promise, SubsystemType Services) {
+		.EmplaceAsyncLambda([&AccountId, &UserStats](FAsyncLambdaResult Promise, const IOnlineServicesPtr& Services) {
 			QueryStats_Fixture(MoveTemp(Promise), Services, AccountId, AccountId, { "Stat_Use_Largest", "Stat_Use_Smallest", "Stat_Use_Sum", "Stat_Use_Set" }, UserStats); 
 		})
-		.EmplaceLambda([&UserStats](SubsystemType Services) {
+		.EmplaceLambda([&UserStats](const IOnlineServicesPtr& Services) {
 			CHECK(UserStats.Find("Stat_Use_Largest") != nullptr);
 			CHECK(UserStats.Find("Stat_Use_Smallest") != nullptr);
 			CHECK(UserStats.Find("Stat_Use_Sum") != nullptr);
 			CHECK(UserStats.Find("Stat_Use_Set") != nullptr);
 		})
-		.EmplaceAsyncLambda([&AccountId, &ImplementedStatsReset](FAsyncLambdaResult Promise, SubsystemType Services) {
+		.EmplaceAsyncLambda([&AccountId, &ImplementedStatsReset](FAsyncLambdaResult Promise, const IOnlineServicesPtr& Services) {
 			Services->GetStatsInterface()->ResetStats({AccountId}).OnComplete([Promise=MoveTemp(Promise), &ImplementedStatsReset](const TOnlineResult<FResetStats>& Result) mutable
 			{
 				CHECK_OP_EQ(Result, Errors::NotImplemented());
@@ -140,7 +140,7 @@ STATS_TEST_CASE("ResetStats works or not implemented", "[tmpExcludedStats]")
 				Promise->SetValue(true);
 			});
 		})
-		.EmplaceAsyncLambda([&AccountId, &UserStats, &ImplementedStatsReset](FAsyncLambdaResult Promise, SubsystemType Services) {
+		.EmplaceAsyncLambda([&AccountId, &UserStats, &ImplementedStatsReset](FAsyncLambdaResult Promise, const IOnlineServicesPtr& Services) {
 			if (!ImplementedStatsReset)
 			{
 				Promise->SetValue(true);
@@ -149,7 +149,7 @@ STATS_TEST_CASE("ResetStats works or not implemented", "[tmpExcludedStats]")
 
 			QueryStats_Fixture(MoveTemp(Promise), Services, AccountId, AccountId, { "Stat_Use_Largest", "Stat_Use_Smallest", "Stat_Use_Sum", "Stat_Use_Set" }, UserStats); 
 		})
-		.EmplaceLambda([&UserStats, &ImplementedStatsReset](SubsystemType Services) {
+		.EmplaceLambda([&UserStats, &ImplementedStatsReset](const IOnlineServicesPtr& Services) {
 			if (!ImplementedStatsReset)
 			{
 				return;
@@ -166,7 +166,7 @@ STATS_TEST_CASE("ResetStats works or not implemented", "[tmpExcludedStats]")
 	RunToCompletion();
 }
 
-void BatchQueryStats_Fixture(FAsyncLambdaResult Promise, SubsystemType Services, const FAccountId& LocalAccountId, const TArray<FAccountId>& AccountIds, TArray<FUserStats>& OutUsersStats)
+void BatchQueryStats_Fixture(FAsyncLambdaResult Promise, const IOnlineServicesPtr& Services, const FAccountId& LocalAccountId, const TArray<FAccountId>& AccountIds, TArray<FUserStats>& OutUsersStats)
 {
 	TArray<FString> StatNames = { "Stat_Use_Largest", "Stat_Use_Smallest", "Stat_Use_Sum", "Stat_Use_Set" };
 	Services->GetStatsInterface()->BatchQueryStats({ LocalAccountId, AccountIds, StatNames }).OnComplete([Promise=MoveTemp(Promise), &OutUsersStats](const TOnlineResult<UE::Online::FBatchQueryStats>& Result) mutable
@@ -199,22 +199,22 @@ STATS_TEST_CASE("BatchUpdateAndQuery with 2 users", "[MultiAccount]")
 	int64 Expected_Stat_Use_Set_B = (int64)FMath::RandRange(1, 100);
 
 	GetLoginPipeline({ AccountIdA, AccountIdB })
-		.EmplaceAsyncLambda([&AccountIdA, &Expected_Stat_Use_Set_A](FAsyncLambdaResult Promise, SubsystemType Services) {
+		.EmplaceAsyncLambda([&AccountIdA, &Expected_Stat_Use_Set_A](FAsyncLambdaResult Promise, const IOnlineServicesPtr& Services) {
 			TArray<FUserStats> UpdateUsersStats{
 				{ AccountIdA, { { "Stat_Use_Set", FStatValue(Expected_Stat_Use_Set_A) } } },
 			};
 			UpdateStats_Fixture(MoveTemp(Promise), Services, AccountIdA, UpdateUsersStats);
 		})
-		.EmplaceAsyncLambda([&AccountIdB, &Expected_Stat_Use_Set_B](FAsyncLambdaResult Promise, SubsystemType Services) {
+		.EmplaceAsyncLambda([&AccountIdB, &Expected_Stat_Use_Set_B](FAsyncLambdaResult Promise, const IOnlineServicesPtr& Services) {
 			TArray<FUserStats> UpdateUsersStats{
 				{ AccountIdB, { { "Stat_Use_Set", FStatValue(Expected_Stat_Use_Set_B) } } },
 			};
 			UpdateStats_Fixture(MoveTemp(Promise), Services, AccountIdB, UpdateUsersStats);
 		})
-		.EmplaceAsyncLambda([&AccountIdA, &AccountIdB, &UpdatedUsersStats](FAsyncLambdaResult Promise, SubsystemType Services) {
+		.EmplaceAsyncLambda([&AccountIdA, &AccountIdB, &UpdatedUsersStats](FAsyncLambdaResult Promise, const IOnlineServicesPtr& Services) {
 			BatchQueryStats_Fixture(MoveTemp(Promise), Services, AccountIdA, { AccountIdA, AccountIdB }, UpdatedUsersStats);
 		})
-		.EmplaceLambda([&AccountIdA, &AccountIdB, &UpdatedUsersStats, &Expected_Stat_Use_Set_A, &Expected_Stat_Use_Set_B](SubsystemType Services) {
+		.EmplaceLambda([&AccountIdA, &AccountIdB, &UpdatedUsersStats, &Expected_Stat_Use_Set_A, &Expected_Stat_Use_Set_B](const IOnlineServicesPtr& Services) {
 			CHECK_EQUAL(StatValueInBatchQueryResult(UpdatedUsersStats, AccountIdA, "Stat_Use_Set"), Expected_Stat_Use_Set_A);
 			CHECK_EQUAL(StatValueInBatchQueryResult(UpdatedUsersStats, AccountIdB, "Stat_Use_Set"), Expected_Stat_Use_Set_B);
 		});
diff --git a/Engine/Source/Programs/Online/OnlineTestsCore/Source/Private/Tests/TestAsyncOp.cpp b/Engine/Source/Programs/Online/OnlineTestsCore/Source/Private/Tests/TestAsyncOp.cpp
--- a/Engine/Source/Programs/Online/OnlineTestsCore/Source/Private/Tests/TestAsyncOp.cpp
+++ b/Engine/Source/Programs/Online/OnlineTestsCore/Source/Private/Tests/TestAsyncOp.cpp
@@ -69,7 +69,7 @@ END_ONLINE_STRUCT_META()
 ASYNCOP_TEST_CASE("Verify that the cache returns the same joinable op when the parameters are the same")
 {
 	FTestPipeline& TestPipeline = GetPipeline();
-	TestPipeline.EmplaceAsyncLambda([](FAsyncLambdaResult Promise, SubsystemType Services)
+	TestPipeline.EmplaceAsyncLambda([](FAsyncLambdaResult Promise, const IOnlineServicesPtr& Services)
 		{
 			UE::Online::FOnlineServicesCommon& ServicesCommon = static_cast<UE::Online::FOnlineServicesCommon&>(*Services);
 			FTestOpWithParams::Params Params1;
@@ -91,7 +91,7 @@ ASYNCOP_TEST_CASE("Verify that the cache returns the same joinable op when the p
 ASYNCOP_TEST_CASE("Verify that the cache returns different joinable ops when the parameters are different")
 {
 	FTestPipeline& TestPipeline = GetPipeline();
-	TestPipeline.EmplaceAsyncLambda([](FAsyncLambdaResult Promise, SubsystemType Services)
+	TestPipeline.EmplaceAsyncLambda([](FAsyncLambdaResult Promise, const IOnlineServicesPtr& Services)
 		{
 			UE::Online::FOnlineServicesCommon& ServicesCommon = static_cast<UE::Online::FOnlineServicesCommon&>(*Services);
 			FTestOpWithParams::Params Params1;
@@ -115,7 +115,7 @@ ASYNCOP_TEST_CASE("Verify that the cache returns different joinable ops when the
 ASYNCOP_TEST_CASE("Verify that the cache returns the same mergeable op when the parameters are the same")
 {
 	FTestPipeline& TestPipeline = GetPipeline();
-	TestPipeline.EmplaceAsyncLambda([](FAsyncLambdaResult Promise, SubsystemType Services)
+	TestPipeline.EmplaceAsyncLambda([](FAsyncLambdaResult Promise, const IOnlineServicesPtr& Services)
 		{
 			UE::Online::FOnlineServicesCommon& ServicesCommon = static_cast<UE::Online::FOnlineServicesCommon&>(*Services);
 			FMergeableTestOpWithParams::Params Params1;
@@ -137,7 +137,7 @@ ASYNCOP_TEST_CASE("Verify that the cache returns the same mergeable op when the
 ASYNCOP_TEST_CASE("Verify that the cache returns different mergeable ops when the parameters are different")
 {
 	FTestPipeline& TestPipeline = GetPipeline();
-	TestPipeline.EmplaceAsyncLambda([](FAsyncLambdaResult Promise, SubsystemType Services)
+	TestPipeline.EmplaceAsyncLambda([](FAsyncLambdaResult Promise, const IOnlineServicesPtr& Services)
 		{
 			UE::Online::FOnlineServicesCommon& ServicesCommon = static_cast<UE::Online::FOnlineServicesCommon&>(*Services);
 			FMergeableTestOpWithParams::Params Params1;
@@ -161,7 +161,7 @@ ASYNCOP_TEST_CASE("Verify that the cache returns different mergeable ops when th
 ASYNCOP_TEST_CASE("Verify that the various OnlineAsyncOp continuation forms complete (Immediately fulfilled promises)")
 {
 	FTestPipeline& TestPipeline = GetPipeline();
-	TestPipeline.EmplaceAsyncLambda([](FAsyncLambdaResult Promise, SubsystemType Services)
+	TestPipeline.EmplaceAsyncLambda([](FAsyncLambdaResult Promise, const IOnlineServicesPtr& Services)
 		{
 			UE::Online::FOnlineServicesCommon& ServicesCommon = static_cast<UE::Online::FOnlineServicesCommon&>(*Services);
 			TOnlineAsyncOpRef<FTestOp> Op = ServicesCommon.GetOp<FTestOp>(FTestOp::Params());
@@ -186,7 +186,7 @@ ASYNCOP_TEST_CASE("Verify that the various OnlineAsyncOp continuation forms comp
 ASYNCOP_TEST_CASE("Verify that the various OnlineAsyncOp continuation forms complete (Promises not immediatly fulfilled)")
 {
 	FTestPipeline& TestPipeline = GetPipeline();
-	TestPipeline.EmplaceAsyncLambda([](FAsyncLambdaResult Promise, SubsystemType Services)
+	TestPipeline.EmplaceAsyncLambda([](FAsyncLambdaResult Promise, const IOnlineServicesPtr& Services)
 		{
 			UE::Online::FOnlineServicesCommon& ServicesCommon = static_cast<UE::Online::FOnlineServicesCommon&>(*Services);
 			TOnlineAsyncOpRef<FTestOp> Op = ServicesCommon.GetOp<FTestOp>(FTestOp::Params());
@@ -212,7 +212,7 @@ ASYNCOP_TEST_CASE("Verify that the various OnlineAsyncOp execution policies work
 {
 	FTestPipeline& TestPipeline = GetPipeline();
 	uint32 ThreadId;
-	TestPipeline.EmplaceAsyncLambda([&ThreadId](FAsyncLambdaResult Promise, SubsystemType Services)
+	TestPipeline.EmplaceAsyncLambda([&ThreadId](FAsyncLambdaResult Promise, const IOnlineServicesPtr& Services)
 		{
 			UE::Online::FOnlineServicesCommon& ServicesCommon = static_cast<UE::Online::FOnlineServicesCommon&>(*Services);
 			TOnlineAsyncOpRef<FTestOp> Op = ServicesCommon.GetOp<FTestOp>(FTestOp::Params());
@@ -259,7 +259,7 @@ ASYNCOP_TEST_CASE("Verify that the duration cache policy works)")
 	// CacheExpiration=Duration
 	// CacheExpirySeconds=2
 	// bCacheError=false
-	TestPipeline.EmplaceLambda([&CacheDurationTestHeiarchy](SubsystemType Services)
+	TestPipeline.EmplaceLambda([&CacheDurationTestHeiarchy](const IOnlineServicesPtr& Services)
 		{
 			UE::Online::FOnlineServicesCommon& ServicesCommon = static_cast<UE::Online::FOnlineServicesCommon&>(*Services);
 			TOnlineAsyncOpRef<FTestOp> Op = ServicesCommon.GetJoinableOp<FTestOp>(FTestOp::Params(), CacheDurationTestHeiarchy);
@@ -267,29 +267,29 @@ ASYNCOP_TEST_CASE("Verify that the duration cache policy works)")
 				.Then([](TOnlineAsyncOp<FTestOp>& Op) { Op.SetResult(FTestOp::Result()); })
 				.Enqueue(ServicesCommon.GetParallelQueue());
 		})
-		.EmplaceLambda([&CacheDurationTestHeiarchy](SubsystemType Services)
+		.EmplaceLambda([&CacheDurationTestHeiarchy](const IOnlineServicesPtr& Services)
 		{
 			UE::Online::FOnlineServicesCommon& ServicesCommon = static_cast<UE::Online::FOnlineServicesCommon&>(*Services);
 			TOnlineAsyncOpRef<FTestOp> Op = ServicesCommon.GetJoinableOp<FTestOp>(FTestOp::Params(), CacheDurationTestHeiarchy);
 			CHECK(Op->IsReady()); // Cache should return the already in flight/completed op
 		})
-		.EmplaceAsyncLambda([](FAsyncLambdaResult Promise, SubsystemType Services)
+		.EmplaceAsyncLambda([](FAsyncLambdaResult Promise, const IOnlineServicesPtr& Services)
 		{
 			// Wait a second into the cache duration
 			FTSTicker::GetCoreTicker().AddTicker(TEXT("OnlineTestTimeout"), 1.0f, [Promise](float) { Promise->SetValue(true); return false; });
 		})
-		.EmplaceLambda([&CacheDurationTestHeiarchy](SubsystemType Services)
+		.EmplaceLambda([&CacheDurationTestHeiarchy](const IOnlineServicesPtr& Services)
 		{
 			UE::Online::FOnlineServicesCommon& ServicesCommon = static_cast<UE::Online::FOnlineServicesCommon&>(*Services);
 			TOnlineAsyncOpRef<FTestOp> Op = ServicesCommon.GetJoinableOp<FTestOp>(FTestOp::Params(), CacheDurationTestHeiarchy);
 			CHECK(Op->IsReady()); // Cache should return the already completed op
 		})
-		.EmplaceAsyncLambda([](FAsyncLambdaResult Promise, SubsystemType Services)
+		.EmplaceAsyncLambda([](FAsyncLambdaResult Promise, const IOnlineServicesPtr& Services)
 		{
 			// Wait until shortly after the cache should expire
 			FTSTicker::GetCoreTicker().AddTicker(TEXT("OnlineTestTimeout"), 2.0f, [Promise](float) { Promise->SetValue(true); return false; });
 		})
-		.EmplaceAsyncLambda([&CacheDurationTestHeiarchy](FAsyncLambdaResult Promise, SubsystemType Services)
+		.EmplaceAsyncLambda([&CacheDurationTestHeiarchy](FAsyncLambdaResult Promise, const IOnlineServicesPtr& Services)
 		{
 			UE::Online::FOnlineServicesCommon& ServicesCommon = static_cast<UE::Online::FOnlineServicesCommon&>(*Services);
 			TOnlineAsyncOpRef<FTestOp> Op = ServicesCommon.GetJoinableOp<FTestOp>(FTestOp::Params(), CacheDurationTestHeiarchy);
@@ -298,7 +298,7 @@ ASYNCOP_TEST_CASE("Verify that the duration cache policy works)")
 				.Then([Promise](TOnlineAsyncOp<FTestOp>& Op) { Op.SetResult(FTestOp::Result()); Promise->SetValue(true); })
 				.Enqueue(ServicesCommon.GetParallelQueue());
 		})
-		.EmplaceAsyncLambda([](FAsyncLambdaResult Promise, SubsystemType Services)
+		.EmplaceAsyncLambda([](FAsyncLambdaResult Promise, const IOnlineServicesPtr& Services)
 		{
 			// Wait until shortly after the cache should expire
 			FTSTicker::GetCoreTicker().AddTicker(TEXT("OnlineTestTimeout"), 3.0f, [Promise](float) { Promise->SetValue(true); return false; });
@@ -323,7 +323,7 @@ ASYNCOP_TEST_CASE("Verify that the cache error setting works)", ASYNCOP_SKIP_EOS
 	// CacheExpiration=Duration
 	// CacheExpirySeconds=1
 	// bCacheError=true
-	TestPipeline.EmplaceAsyncLambda([&DontCacheErrorHeiarchy](FAsyncLambdaResult Promise, SubsystemType Services)
+	TestPipeline.EmplaceAsyncLambda([&DontCacheErrorHeiarchy](FAsyncLambdaResult Promise, const IOnlineServicesPtr& Services)
 		{
 			UE::Online::FOnlineServicesCommon& ServicesCommon = static_cast<UE::Online::FOnlineServicesCommon&>(*Services);
 			TOnlineAsyncOpRef<FTestOp> Op = ServicesCommon.GetJoinableOp<FTestOp>(FTestOp::Params(), DontCacheErrorHeiarchy);
@@ -331,7 +331,7 @@ ASYNCOP_TEST_CASE("Verify that the cache error setting works)", ASYNCOP_SKIP_EOS
 				.Then([Promise](TOnlineAsyncOp<FTestOp>& Op) { Op.SetError(UE::Online::Errors::Unknown()); Promise->SetValue(true); })
 				.Enqueue(ServicesCommon.GetParallelQueue());
 		})
-		.EmplaceAsyncLambda([&CacheErrorHeiarchy](FAsyncLambdaResult Promise, SubsystemType Services)
+		.EmplaceAsyncLambda([&CacheErrorHeiarchy](FAsyncLambdaResult Promise, const IOnlineServicesPtr& Services)
 		{
 			UE::Online::FOnlineServicesCommon& ServicesCommon = static_cast<UE::Online::FOnlineServicesCommon&>(*Services);
 			TOnlineAsyncOpRef<FTestOp> Op = ServicesCommon.GetJoinableOp<FTestOp>(FTestOp::Params(), CacheErrorHeiarchy);
@@ -340,7 +340,7 @@ ASYNCOP_TEST_CASE("Verify that the cache error setting works)", ASYNCOP_SKIP_EOS
 				.Then([Promise](TOnlineAsyncOp<FTestOp>& Op) { Op.SetError(UE::Online::Errors::Unknown()); Promise->SetValue(true); })
 				.Enqueue(ServicesCommon.GetParallelQueue());
 		})
-		.EmplaceAsyncLambda([&CacheErrorHeiarchy](FAsyncLambdaResult Promise, SubsystemType Services)
+		.EmplaceAsyncLambda([&CacheErrorHeiarchy](FAsyncLambdaResult Promise, const IOnlineServicesPtr& Services)
 		{
 			UE::Online::FOnlineServicesCommon& ServicesCommon = static_cast<UE::Online::FOnlineServicesCommon&>(*Services);
 			TOnlineAsyncOpRef<FTestOp> Op = ServicesCommon.GetJoinableOp<FTestOp>(FTestOp::Params(), CacheErrorHeiarchy);
@@ -351,7 +351,7 @@ ASYNCOP_TEST_CASE("Verify that the cache error setting works)", ASYNCOP_SKIP_EOS
 					Promise->SetValue(true);
 				});
 		})
-		.EmplaceAsyncLambda([](FAsyncLambdaResult Promise, SubsystemType Services)
+		.EmplaceAsyncLambda([](FAsyncLambdaResult Promise, const IOnlineServicesPtr& Services)
 		{
 			// Wait until shortly after the cache should expire
 			FTSTicker::GetCoreTicker().AddTicker(TEXT("OnlineTestTimeout"), 2.0f, [Promise](float) { Promise->SetValue(true); return false; });
@@ -368,7 +368,7 @@ ASYNCOP_TEST_CASE("Verify that the cache expires once complete)")
 	// [OnlineTests.TestCacheExpireUponCompletion]
 	// CacheExpiration=UponCompletion
 	// bCacheError=false
-	TestPipeline.EmplaceAsyncLambda([&CacheExpireUponCompletionHeiarchy](FAsyncLambdaResult Promise, SubsystemType Services)
+	TestPipeline.EmplaceAsyncLambda([&CacheExpireUponCompletionHeiarchy](FAsyncLambdaResult Promise, const IOnlineServicesPtr& Services)
 		{
 			UE::Online::FOnlineServicesCommon& ServicesCommon = static_cast<UE::Online::FOnlineServicesCommon&>(*Services);
 			TOnlineAsyncOpRef<FTestOp> Op = ServicesCommon.GetJoinableOp<FTestOp>(FTestOp::Params(), CacheExpireUponCompletionHeiarchy);
@@ -379,7 +379,7 @@ ASYNCOP_TEST_CASE("Verify that the cache expires once complete)")
 			CHECK(Op2->IsReady());
 			CHECK(!Op2->IsComplete());
 		})
-		.EmplaceAsyncLambda([&CacheExpireUponCompletionHeiarchy](FAsyncLambdaResult Promise, SubsystemType Services)
+		.EmplaceAsyncLambda([&CacheExpireUponCompletionHeiarchy](FAsyncLambdaResult Promise, const IOnlineServicesPtr& Services)
 		{
 			UE::Online::FOnlineServicesCommon& ServicesCommon = static_cast<UE::Online::FOnlineServicesCommon&>(*Services);
 			TOnlineAsyncOpRef<FTestOp> Op = ServicesCommon.GetJoinableOp<FTestOp>(FTestOp::Params(), CacheExpireUponCompletionHeiarchy);
diff --git a/Engine/Source/Programs/Online/OnlineTestsCore/Source/Public/AsyncTestStep.h b/Engine/Source/Programs/Online/OnlineTestsCore/Source/Public/AsyncTestStep.h
--- a/Engine/Source/Programs/Online/OnlineTestsCore/Source/Public/AsyncTestStep.h
+++ b/Engine/Source/Programs/Online/OnlineTestsCore/Source/Public/AsyncTestStep.h
@@ -43,9 +43,9 @@ class FAsyncTestStep : public FTestPipeline::FStep
 {
 public:
 	// main entry point function- OnlineSubsystem will be V2. The future bool in Result is whether or not to continue the test (essentually just a require)
-	virtual void Run(FAsyncStepResult Result, SubsystemType OnlineSubsystem) = 0;
+	virtual void Run(FAsyncStepResult Result, const IOnlineServicesPtr& OnlineServices) = 0;
 
-	virtual EContinuance Tick(SubsystemType OnlineSubsystem) override
+	virtual EContinuance Tick(const IOnlineServicesPtr& OnlineServices) override
 	{
 		if (!bInitialized)
 		{
@@ -56,7 +56,7 @@ public:
 				REQUIRE(Result);
 				bComplete = true;
 			});
-			Run(Result, OnlineSubsystem);
+			Run(Result, OnlineServices);
 		}
 
 		if (bComplete == true)
diff --git a/Engine/Source/Programs/Online/OnlineTestsCore/Source/Public/AuthAutoLogin.h b/Engine/Source/Programs/Online/OnlineTestsCore/Source/Public/AuthAutoLogin.h
--- a/Engine/Source/Programs/Online/OnlineTestsCore/Source/Public/AuthAutoLogin.h
+++ b/Engine/Source/Programs/Online/OnlineTestsCore/Source/Public/AuthAutoLogin.h
@@ -26,7 +26,7 @@ struct FAuthAutoLoginStep : public FAsyncTestStep
 	{
 	};
 
-	virtual void Run(FAsyncStepResult Promise, SubsystemType Services) override
+	virtual void Run(FAsyncStepResult Promise, const IOnlineServicesPtr& Services) override
 	{
 		UE::Online::IAuthPtr OnlineAuthPtr = Services->GetAuthInterface();
 
diff --git a/Engine/Source/Programs/Online/OnlineTestsCore/Source/Public/Helpers/Achievements/CheckAchievementStateHelper.h b/Engine/Source/Programs/Online/OnlineTestsCore/Source/Public/Helpers/Achievements/CheckAchievementStateHelper.h
--- a/Engine/Source/Programs/Online/OnlineTestsCore/Source/Public/Helpers/Achievements/CheckAchievementStateHelper.h
+++ b/Engine/Source/Programs/Online/OnlineTestsCore/Source/Public/Helpers/Achievements/CheckAchievementStateHelper.h
@@ -6,7 +6,7 @@
 
 namespace UE::Online
 {
-	inline void CheckAchievementState_Fixture(FAsyncLambdaResult Promise, SubsystemType Services, const FAccountId& AccountId, const FString& AchievementId, float AchievementProgress)
+	inline void CheckAchievementState_Fixture(const FAsyncLambdaResult& Promise, const IOnlineServicesPtr& Services, const FAccountId& AccountId, const FString& AchievementId, float AchievementProgress)
 	{
 		FQueryAchievementStates::Params Params{ AccountId };
 		Services->GetAchievementsInterface()->QueryAchievementStates(MoveTemp(Params)).OnComplete([&](const TOnlineResult<UE::Online::FQueryAchievementStates>& Result) mutable
diff --git a/Engine/Source/Programs/Online/OnlineTestsCore/Source/Public/Helpers/Achievements/QueryAchievementDefinitionsHelper.h b/Engine/Source/Programs/Online/OnlineTestsCore/Source/Public/Helpers/Achievements/QueryAchievementDefinitionsHelper.h
--- a/Engine/Source/Programs/Online/OnlineTestsCore/Source/Public/Helpers/Achievements/QueryAchievementDefinitionsHelper.h
+++ b/Engine/Source/Programs/Online/OnlineTestsCore/Source/Public/Helpers/Achievements/QueryAchievementDefinitionsHelper.h
@@ -29,7 +29,7 @@ struct FQueryAchievementDefinitionsHelper : public FTestPipeline::FStep
 
 	enum class EState { Init, Called, Done } State = EState::Init;
 
-	virtual EContinuance Tick(SubsystemType OnlineSubsystem) override
+	virtual EContinuance Tick(const IOnlineServicesPtr& OnlineSubsystem) override
 	{
 		switch (State)
 		{
diff --git a/Engine/Source/Programs/Online/OnlineTestsCore/Source/Public/Helpers/Achievements/QueryAchievementStatesHelper.h b/Engine/Source/Programs/Online/OnlineTestsCore/Source/Public/Helpers/Achievements/QueryAchievementStatesHelper.h
--- a/Engine/Source/Programs/Online/OnlineTestsCore/Source/Public/Helpers/Achievements/QueryAchievementStatesHelper.h
+++ b/Engine/Source/Programs/Online/OnlineTestsCore/Source/Public/Helpers/Achievements/QueryAchievementStatesHelper.h
@@ -30,7 +30,7 @@ struct FQueryAchievementStatesHelper : public FTestPipeline::FStep
 
 	enum class EState { Init, Called, Done } State = EState::Init;
 
-	virtual EContinuance Tick(SubsystemType OnlineSubsystem) override
+	virtual EContinuance Tick(const IOnlineServicesPtr& OnlineSubsystem) override
 	{
 		switch (State)
 		{
diff --git a/Engine/Source/Programs/Online/OnlineTestsCore/Source/Public/Helpers/Achievements/UnlockAchievementsHelper.h b/Engine/Source/Programs/Online/OnlineTestsCore/Source/Public/Helpers/Achievements/UnlockAchievementsHelper.h
--- a/Engine/Source/Programs/Online/OnlineTestsCore/Source/Public/Helpers/Achievements/UnlockAchievementsHelper.h
+++ b/Engine/Source/Programs/Online/OnlineTestsCore/Source/Public/Helpers/Achievements/UnlockAchievementsHelper.h
@@ -30,7 +30,7 @@ struct FUnlockAchievementsHelper : public FTestPipeline::FStep
 
 	enum class EState { Init, Called, AwaitEvent, Done } State = EState::Init;
 
-	virtual EContinuance Tick(SubsystemType OnlineSubsystem) override
+	virtual EContinuance Tick(const IOnlineServicesPtr& OnlineSubsystem) override
 	{
 		switch (State)
 		{
diff --git a/Engine/Source/Programs/Online/OnlineTestsCore/Source/Public/Helpers/Auth/AuthLogin.h b/Engine/Source/Programs/Online/OnlineTestsCore/Source/Public/Helpers/Auth/AuthLogin.h
--- a/Engine/Source/Programs/Online/OnlineTestsCore/Source/Public/Helpers/Auth/AuthLogin.h
+++ b/Engine/Source/Programs/Online/OnlineTestsCore/Source/Public/Helpers/Auth/AuthLogin.h
@@ -21,7 +21,7 @@ struct FAuthLoginStep : public FAsyncTestStep
 
 	virtual ~FAuthLoginStep() = default;
 
-	virtual void Run(FAsyncStepResult Promise, SubsystemType Services) override
+	virtual void Run(FAsyncStepResult Promise, const IOnlineServicesPtr& Services) override
 	{
 		OnlineAuthPtr = Services->GetAuthInterface();
 		REQUIRE(OnlineAuthPtr);
diff --git a/Engine/Source/Programs/Online/OnlineTestsCore/Source/Public/Helpers/Auth/AuthLogout.h b/Engine/Source/Programs/Online/OnlineTestsCore/Source/Public/Helpers/Auth/AuthLogout.h
--- a/Engine/Source/Programs/Online/OnlineTestsCore/Source/Public/Helpers/Auth/AuthLogout.h
+++ b/Engine/Source/Programs/Online/OnlineTestsCore/Source/Public/Helpers/Auth/AuthLogout.h
@@ -2,10 +2,11 @@
 #pragma once
 
 #include "CoreMinimal.h"
+#include "AsyncTestStep.h"
+#include "OnlineCatchHelper.h"
 #include "TestDriver.h"
 #include "TestHarness.h"
-#include "AsyncTestStep.h"
-
+#include "Online/Auth.h"
 #include "Online/AuthCommon.h"
 #include "Online/OnlineAsyncOp.h"
 #include "Online/OnlineErrorDefinitions.h"
@@ -24,7 +25,7 @@ struct FAuthLogoutStep : public FAsyncTestStep
 
 	virtual ~FAuthLogoutStep() = default;
 
-	virtual void Run(FAsyncStepResult Promise, SubsystemType Services) override
+	virtual void Run(FAsyncStepResult Promise, const IOnlineServicesPtr& Services) override
 	{
 		OnlineAuthPtr = Services->GetAuthInterface();
 		REQUIRE(OnlineAuthPtr);
@@ -34,14 +35,14 @@ struct FAuthLogoutStep : public FAsyncTestStep
 			PlatformUserId = *PlatformUserIdPtr;
 		}
 		
-		TOnlineResult<FAuthGetLocalOnlineUserByPlatformUserId> AccountId = OnlineAuthPtr->GetLocalOnlineUserByPlatformUserId({ PlatformUserId });
+		UE::Online::TOnlineResult<UE::Online::FAuthGetLocalOnlineUserByPlatformUserId> AccountId = OnlineAuthPtr->GetLocalOnlineUserByPlatformUserId({ PlatformUserId });
 
 		CAPTURE(ToLogString(AccountId), PlatformUserId);
 		CHECK_OP(AccountId);
 		if(AccountId.IsOk())
 		{
 			OnlineAuthPtr->Logout({ AccountId.GetOkValue().AccountInfo->AccountId})
-				.OnComplete([this, Promise = MoveTemp(Promise)](const TOnlineResult<FAuthLogout> Op) mutable
+				.OnComplete([this, Promise = MoveTemp(Promise)](const UE::Online::TOnlineResult<UE::Online::FAuthLogout> Op) mutable
 				{
 					CHECK_OP_EQ(Op, Errors::NotImplemented());
 					Promise->SetValue(true);
diff --git a/Engine/Source/Programs/Online/OnlineTestsCore/Source/Public/Helpers/Auth/AuthQueryExternalAuthToken.h b/Engine/Source/Programs/Online/OnlineTestsCore/Source/Public/Helpers/Auth/AuthQueryExternalAuthToken.h
--- a/Engine/Source/Programs/Online/OnlineTestsCore/Source/Public/Helpers/Auth/AuthQueryExternalAuthToken.h
+++ b/Engine/Source/Programs/Online/OnlineTestsCore/Source/Public/Helpers/Auth/AuthQueryExternalAuthToken.h
@@ -29,7 +29,7 @@ struct FAuthQueryExternalAuthTokenStep : public FAsyncTestStep
 		}
 	};
 
-	virtual void Run(FAsyncStepResult Promise, SubsystemType Services) override
+	virtual void Run(FAsyncStepResult Promise, const IOnlineServicesPtr& Services) override
 	{
 		OnlineAuthPtr = Services->GetAuthInterface();
 		REQUIRE(OnlineAuthPtr);
diff --git a/Engine/Source/Programs/Online/OnlineTestsCore/Source/Public/Helpers/Commerce/CheckoutHelper.h b/Engine/Source/Programs/Online/OnlineTestsCore/Source/Public/Helpers/Commerce/CheckoutHelper.h
--- a/Engine/Source/Programs/Online/OnlineTestsCore/Source/Public/Helpers/Commerce/CheckoutHelper.h
+++ b/Engine/Source/Programs/Online/OnlineTestsCore/Source/Public/Helpers/Commerce/CheckoutHelper.h
@@ -26,7 +26,7 @@ struct FCheckoutHelper : public FAsyncTestStep
 
 	virtual ~FCheckoutHelper() = default;
 
-	virtual void Run(FAsyncStepResult Promise, SubsystemType Services) override
+	virtual void Run(FAsyncStepResult Promise, const IOnlineServicesPtr& Services) override
 	{
 		CommerceInterface = Services->GetCommerceInterface();
 		REQUIRE(CommerceInterface);
diff --git a/Engine/Source/Programs/Online/OnlineTestsCore/Source/Public/Helpers/Commerce/GetEntitlementsHelper.h b/Engine/Source/Programs/Online/OnlineTestsCore/Source/Public/Helpers/Commerce/GetEntitlementsHelper.h
--- a/Engine/Source/Programs/Online/OnlineTestsCore/Source/Public/Helpers/Commerce/GetEntitlementsHelper.h
+++ b/Engine/Source/Programs/Online/OnlineTestsCore/Source/Public/Helpers/Commerce/GetEntitlementsHelper.h
@@ -27,7 +27,7 @@ struct FGetEntitlementsHelper : public FAsyncTestStep
 
 	virtual ~FGetEntitlementsHelper() = default;
 
-	virtual void Run(FAsyncStepResult Promise, SubsystemType Services) override
+	virtual void Run(FAsyncStepResult Promise, const IOnlineServicesPtr& Services) override
 	{
 		CommerceInterface = Services->GetCommerceInterface();
 		REQUIRE(CommerceInterface);
diff --git a/Engine/Source/Programs/Online/OnlineTestsCore/Source/Public/Helpers/Commerce/GetOffersByIdHelper.h b/Engine/Source/Programs/Online/OnlineTestsCore/Source/Public/Helpers/Commerce/GetOffersByIdHelper.h
--- a/Engine/Source/Programs/Online/OnlineTestsCore/Source/Public/Helpers/Commerce/GetOffersByIdHelper.h
+++ b/Engine/Source/Programs/Online/OnlineTestsCore/Source/Public/Helpers/Commerce/GetOffersByIdHelper.h
@@ -27,7 +27,7 @@ struct FGetOffersByIdHelper : public FAsyncTestStep
 
 	virtual ~FGetOffersByIdHelper() = default;
 
-	virtual void Run(FAsyncStepResult Promise, SubsystemType Services) override
+	virtual void Run(FAsyncStepResult Promise, const IOnlineServicesPtr& Services) override
 	{
 		CommerceInterface = Services->GetCommerceInterface();
 		REQUIRE(CommerceInterface);
diff --git a/Engine/Source/Programs/Online/OnlineTestsCore/Source/Public/Helpers/Commerce/GetOffersHelper.h b/Engine/Source/Programs/Online/OnlineTestsCore/Source/Public/Helpers/Commerce/GetOffersHelper.h
--- a/Engine/Source/Programs/Online/OnlineTestsCore/Source/Public/Helpers/Commerce/GetOffersHelper.h
+++ b/Engine/Source/Programs/Online/OnlineTestsCore/Source/Public/Helpers/Commerce/GetOffersHelper.h
@@ -27,7 +27,7 @@ struct FGetOffersHelper : public FAsyncTestStep
 
 	virtual ~FGetOffersHelper() = default;
 
-	virtual void Run(FAsyncStepResult Promise, SubsystemType Services) override
+	virtual void Run(FAsyncStepResult Promise, const IOnlineServicesPtr& Services) override
 	{
 		CommerceInterface = Services->GetCommerceInterface();
 		REQUIRE(CommerceInterface);
diff --git a/Engine/Source/Programs/Online/OnlineTestsCore/Source/Public/Helpers/Commerce/QueryEntitlementsHelper.h b/Engine/Source/Programs/Online/OnlineTestsCore/Source/Public/Helpers/Commerce/QueryEntitlementsHelper.h
--- a/Engine/Source/Programs/Online/OnlineTestsCore/Source/Public/Helpers/Commerce/QueryEntitlementsHelper.h
+++ b/Engine/Source/Programs/Online/OnlineTestsCore/Source/Public/Helpers/Commerce/QueryEntitlementsHelper.h
@@ -26,7 +26,7 @@ struct FQueryEntitlementsHelper : public FAsyncTestStep
 
 	virtual ~FQueryEntitlementsHelper() = default;
 
-	virtual void Run(FAsyncStepResult Promise, SubsystemType Services) override
+	virtual void Run(FAsyncStepResult Promise, const IOnlineServicesPtr& Services) override
 	{
 		CommerceInterface = Services->GetCommerceInterface();
 		REQUIRE(CommerceInterface);
diff --git a/Engine/Source/Programs/Online/OnlineTestsCore/Source/Public/Helpers/Commerce/QueryOffersByIdHelper.h b/Engine/Source/Programs/Online/OnlineTestsCore/Source/Public/Helpers/Commerce/QueryOffersByIdHelper.h
--- a/Engine/Source/Programs/Online/OnlineTestsCore/Source/Public/Helpers/Commerce/QueryOffersByIdHelper.h
+++ b/Engine/Source/Programs/Online/OnlineTestsCore/Source/Public/Helpers/Commerce/QueryOffersByIdHelper.h
@@ -26,7 +26,7 @@ struct FQueryOffersByIdHelper : public FAsyncTestStep
 
 	virtual ~FQueryOffersByIdHelper() = default;
 
-	virtual void Run(FAsyncStepResult Promise, SubsystemType Services) override
+	virtual void Run(FAsyncStepResult Promise, const IOnlineServicesPtr& Services) override
 	{
 		CommerceInterface = Services->GetCommerceInterface();
 		REQUIRE(CommerceInterface);
diff --git a/Engine/Source/Programs/Online/OnlineTestsCore/Source/Public/Helpers/Commerce/QueryOffersHelper.h b/Engine/Source/Programs/Online/OnlineTestsCore/Source/Public/Helpers/Commerce/QueryOffersHelper.h
--- a/Engine/Source/Programs/Online/OnlineTestsCore/Source/Public/Helpers/Commerce/QueryOffersHelper.h
+++ b/Engine/Source/Programs/Online/OnlineTestsCore/Source/Public/Helpers/Commerce/QueryOffersHelper.h
@@ -26,7 +26,7 @@ struct FQueryOffersHelper : public FAsyncTestStep
 
 	virtual ~FQueryOffersHelper() = default;
 
-	virtual void Run(FAsyncStepResult Promise, SubsystemType Services) override
+	virtual void Run(FAsyncStepResult Promise, const IOnlineServicesPtr& Services) override
 	{
 		CommerceInterface = Services->GetCommerceInterface();
 		REQUIRE(CommerceInterface);
diff --git a/Engine/Source/Programs/Online/OnlineTestsCore/Source/Public/Helpers/Commerce/RedeemEntitlementsHelper.h b/Engine/Source/Programs/Online/OnlineTestsCore/Source/Public/Helpers/Commerce/RedeemEntitlementsHelper.h
--- a/Engine/Source/Programs/Online/OnlineTestsCore/Source/Public/Helpers/Commerce/RedeemEntitlementsHelper.h
+++ b/Engine/Source/Programs/Online/OnlineTestsCore/Source/Public/Helpers/Commerce/RedeemEntitlementsHelper.h
@@ -28,7 +28,7 @@ struct FRedeemEntitlementsHelper : public FAsyncTestStep
 
 	virtual ~FRedeemEntitlementsHelper() = default;
 
-	virtual void Run(FAsyncStepResult Promise, SubsystemType Services) override
+	virtual void Run(FAsyncStepResult Promise, const IOnlineServicesPtr& Services) override
 	{
 		HelperParams.OpParams->EntitlementId = *EntitlementId;
 
diff --git a/Engine/Source/Programs/Online/OnlineTestsCore/Source/Public/Helpers/Identity/IdentityGetLoginByUserId.h b/Engine/Source/Programs/Online/OnlineTestsCore/Source/Public/Helpers/Identity/IdentityGetLoginByUserId.h
--- a/Engine/Source/Programs/Online/OnlineTestsCore/Source/Public/Helpers/Identity/IdentityGetLoginByUserId.h
+++ b/Engine/Source/Programs/Online/OnlineTestsCore/Source/Public/Helpers/Identity/IdentityGetLoginByUserId.h
@@ -27,7 +27,7 @@ struct FIdentityGetUniquePlayerIdStep : public FTestPipeline::FStep
 
 	virtual ~FIdentityGetUniquePlayerIdStep() = default;
 
-	virtual EContinuance Tick(SubsystemType OnlineSubsystem) override
+	virtual EContinuance Tick(const IOnlineServicesPtr& OnlineSubsystem) override
 	{
 		UE::Online::IAuthPtr OnlineAuthPtr = OnlineSubsystem->GetAuthInterface();
 
diff --git a/Engine/Source/Programs/Online/OnlineTestsCore/Source/Public/Helpers/LambdaStep.h b/Engine/Source/Programs/Online/OnlineTestsCore/Source/Public/Helpers/LambdaStep.h
--- a/Engine/Source/Programs/Online/OnlineTestsCore/Source/Public/Helpers/LambdaStep.h
+++ b/Engine/Source/Programs/Online/OnlineTestsCore/Source/Public/Helpers/LambdaStep.h
@@ -5,19 +5,19 @@
 
 struct FLambdaStep : public FTestPipeline::FStep
 {
-	FLambdaStep(TFunction<void(SubsystemType)>&& InFunction)
+	FLambdaStep(TFunction<void(IOnlineServicesPtr)>&& InFunction)
 		: Function(MoveTemp(InFunction))
 	{
 	}
 
 	virtual ~FLambdaStep() = default;
 
-	virtual EContinuance Tick(SubsystemType OnlineSubsystem) override
+	virtual EContinuance Tick(const IOnlineServicesPtr& OnlineSubsystem) override
 	{
 		Function(OnlineSubsystem);
 		return EContinuance::Done;
 	}
 
 protected:
-	TFunction<void(SubsystemType)> Function;
+	TFunction<void(IOnlineServicesPtr)> Function;
 };
\ No newline at end of file
diff --git a/Engine/Source/Programs/Online/OnlineTestsCore/Source/Public/Helpers/Lobby/LobbyCreateHelper.h b/Engine/Source/Programs/Online/OnlineTestsCore/Source/Public/Helpers/Lobby/LobbyCreateHelper.h
--- a/Engine/Source/Programs/Online/OnlineTestsCore/Source/Public/Helpers/Lobby/LobbyCreateHelper.h
+++ b/Engine/Source/Programs/Online/OnlineTestsCore/Source/Public/Helpers/Lobby/LobbyCreateHelper.h
@@ -28,7 +28,7 @@ struct FLobbyCreateHelper : public FTestPipeline::FStep
 
 	enum class EState { Init, CreateLobbyCalled, CreateLobbyInProgress, CreateLobbyComplete, Done } State = EState::Init;
 
-	virtual EContinuance Tick(SubsystemType OnlineSubsystem) override
+	virtual EContinuance Tick(const IOnlineServicesPtr& OnlineSubsystem) override
 	{
 		switch (State)
 		{
diff --git a/Engine/Source/Programs/Online/OnlineTestsCore/Source/Public/Helpers/Lobby/LobbyJoinHelper.h b/Engine/Source/Programs/Online/OnlineTestsCore/Source/Public/Helpers/Lobby/LobbyJoinHelper.h
--- a/Engine/Source/Programs/Online/OnlineTestsCore/Source/Public/Helpers/Lobby/LobbyJoinHelper.h
+++ b/Engine/Source/Programs/Online/OnlineTestsCore/Source/Public/Helpers/Lobby/LobbyJoinHelper.h
@@ -34,7 +34,7 @@ struct FLobbyJoinHelper : public FAsyncTestStep
 
 	virtual ~FLobbyJoinHelper() = default;
 
-	virtual void Run(FAsyncStepResult Promise, SubsystemType Services) override
+	virtual void Run(FAsyncStepResult Promise, const IOnlineServicesPtr& Services) override
 	{
 		LobbyInterface = Services->GetLobbiesInterface();
 		REQUIRE(LobbyInterface);
diff --git a/Engine/Source/Programs/Online/OnlineTestsCore/Source/Public/Helpers/Presence/BatchQueryPresenceHelper.h b/Engine/Source/Programs/Online/OnlineTestsCore/Source/Public/Helpers/Presence/BatchQueryPresenceHelper.h
--- a/Engine/Source/Programs/Online/OnlineTestsCore/Source/Public/Helpers/Presence/BatchQueryPresenceHelper.h
+++ b/Engine/Source/Programs/Online/OnlineTestsCore/Source/Public/Helpers/Presence/BatchQueryPresenceHelper.h
@@ -18,7 +18,7 @@ public:
 
 	}
 
-	virtual void Run(FAsyncStepResult Promise, SubsystemType Services) override
+	virtual void Run(FAsyncStepResult Promise, const IOnlineServicesPtr& Services) override
 	{
 		IPresencePtr Presence = Services->GetPresenceInterface();
 
diff --git a/Engine/Source/Programs/Online/OnlineTestsCore/Source/Public/Helpers/Presence/CommonPresenceHelpers.h b/Engine/Source/Programs/Online/OnlineTestsCore/Source/Public/Helpers/Presence/CommonPresenceHelpers.h
--- a/Engine/Source/Programs/Online/OnlineTestsCore/Source/Public/Helpers/Presence/CommonPresenceHelpers.h
+++ b/Engine/Source/Programs/Online/OnlineTestsCore/Source/Public/Helpers/Presence/CommonPresenceHelpers.h
@@ -101,7 +101,7 @@ public:
 
 	}
 	
-	virtual void Run(FAsyncStepResult Promise, SubsystemType Services) override
+	virtual void Run(FAsyncStepResult Promise, const IOnlineServicesPtr& Services) override
 	{
 		CHECK(Presence1.IsValid());
 
diff --git a/Engine/Source/Programs/Online/OnlineTestsCore/Source/Public/Helpers/Presence/GetCachedPresenceHelper.h b/Engine/Source/Programs/Online/OnlineTestsCore/Source/Public/Helpers/Presence/GetCachedPresenceHelper.h
--- a/Engine/Source/Programs/Online/OnlineTestsCore/Source/Public/Helpers/Presence/GetCachedPresenceHelper.h
+++ b/Engine/Source/Programs/Online/OnlineTestsCore/Source/Public/Helpers/Presence/GetCachedPresenceHelper.h
@@ -17,7 +17,7 @@ public:
 	}
 
 
-	virtual void Run(FAsyncStepResult Promise, SubsystemType Services) override
+	virtual void Run(FAsyncStepResult Promise, const IOnlineServicesPtr& Services) override
 	{
 		IPresencePtr Presence = Services->GetPresenceInterface();
 
diff --git a/Engine/Source/Programs/Online/OnlineTestsCore/Source/Public/Helpers/Presence/PartialUpdatePresenceHelper.h b/Engine/Source/Programs/Online/OnlineTestsCore/Source/Public/Helpers/Presence/PartialUpdatePresenceHelper.h
--- a/Engine/Source/Programs/Online/OnlineTestsCore/Source/Public/Helpers/Presence/PartialUpdatePresenceHelper.h
+++ b/Engine/Source/Programs/Online/OnlineTestsCore/Source/Public/Helpers/Presence/PartialUpdatePresenceHelper.h
@@ -15,7 +15,7 @@ public:
 
 	}
 
-	virtual void Run(FAsyncStepResult Promise, SubsystemType Services) override
+	virtual void Run(FAsyncStepResult Promise, const IOnlineServicesPtr& Services) override
 	{
 		IPresencePtr Presence = Services->GetPresenceInterface();
 
diff --git a/Engine/Source/Programs/Online/OnlineTestsCore/Source/Public/Helpers/Presence/QueryPresenceHelper.h b/Engine/Source/Programs/Online/OnlineTestsCore/Source/Public/Helpers/Presence/QueryPresenceHelper.h
--- a/Engine/Source/Programs/Online/OnlineTestsCore/Source/Public/Helpers/Presence/QueryPresenceHelper.h
+++ b/Engine/Source/Programs/Online/OnlineTestsCore/Source/Public/Helpers/Presence/QueryPresenceHelper.h
@@ -18,7 +18,7 @@ public:
 		
 	}
 
-	virtual void Run(FAsyncStepResult Promise, SubsystemType Services) override
+	virtual void Run(FAsyncStepResult Promise, const IOnlineServicesPtr& Services) override
 	{
 		IPresencePtr Presence = Services->GetPresenceInterface();
 		Presence->QueryPresence({LocalAccountId, TargetAccountId, bListenToChanges})
diff --git a/Engine/Source/Programs/Online/OnlineTestsCore/Source/Public/Helpers/Presence/UpdatePresenceHelper.h b/Engine/Source/Programs/Online/OnlineTestsCore/Source/Public/Helpers/Presence/UpdatePresenceHelper.h
--- a/Engine/Source/Programs/Online/OnlineTestsCore/Source/Public/Helpers/Presence/UpdatePresenceHelper.h
+++ b/Engine/Source/Programs/Online/OnlineTestsCore/Source/Public/Helpers/Presence/UpdatePresenceHelper.h
@@ -16,7 +16,7 @@ public:
 
 	}
 
-	virtual void Run(FAsyncStepResult Promise, SubsystemType Services) override
+	virtual void Run(FAsyncStepResult Promise, const IOnlineServicesPtr& Services) override
 	{
 		IPresencePtr PresencePtr = Services->GetPresenceInterface();
 		PresencePtr->UpdatePresence({ LocalAccountId, Presence })
diff --git a/Engine/Source/Programs/Online/OnlineTestsCore/Source/Public/Helpers/Sessions/AddRemoveSessionMemberHelper.h b/Engine/Source/Programs/Online/OnlineTestsCore/Source/Public/Helpers/Sessions/AddRemoveSessionMemberHelper.h
--- a/Engine/Source/Programs/Online/OnlineTestsCore/Source/Public/Helpers/Sessions/AddRemoveSessionMemberHelper.h
+++ b/Engine/Source/Programs/Online/OnlineTestsCore/Source/Public/Helpers/Sessions/AddRemoveSessionMemberHelper.h
@@ -26,7 +26,7 @@ struct FAddSessionMemberHelper : public FAsyncTestStep
 
 	virtual ~FAddSessionMemberHelper() = default;
 
-	virtual void Run(FAsyncStepResult Promise, SubsystemType Services) override
+	virtual void Run(FAsyncStepResult Promise, const IOnlineServicesPtr& Services) override
 	{
 		SessionsInterface = Services->GetSessionsInterface();
 		REQUIRE(SessionsInterface);
@@ -75,7 +75,7 @@ struct FRemoveSessionMemberHelper : public FAsyncTestStep
 
 	virtual ~FRemoveSessionMemberHelper() = default;
 
-	virtual void Run(FAsyncStepResult Promise, SubsystemType Services) override
+	virtual void Run(FAsyncStepResult Promise, const IOnlineServicesPtr& Services) override
 	{
 		SessionsInterface = Services->GetSessionsInterface();
 		REQUIRE(SessionsInterface);
diff --git a/Engine/Source/Programs/Online/OnlineTestsCore/Source/Public/Helpers/Sessions/CreateSessionHelper.h b/Engine/Source/Programs/Online/OnlineTestsCore/Source/Public/Helpers/Sessions/CreateSessionHelper.h
--- a/Engine/Source/Programs/Online/OnlineTestsCore/Source/Public/Helpers/Sessions/CreateSessionHelper.h
+++ b/Engine/Source/Programs/Online/OnlineTestsCore/Source/Public/Helpers/Sessions/CreateSessionHelper.h
@@ -25,7 +25,7 @@ struct FCreateSessionHelper : public FAsyncTestStep
 
 	virtual ~FCreateSessionHelper() = default;
 
-	virtual void Run(FAsyncStepResult Promise, SubsystemType Services) override
+	virtual void Run(FAsyncStepResult Promise, const IOnlineServicesPtr& Services) override
 	{
 		SessionsInterface = Services->GetSessionsInterface();
 		REQUIRE(SessionsInterface);
diff --git a/Engine/Source/Programs/Online/OnlineTestsCore/Source/Public/Helpers/Sessions/FindSessionsHelper.h b/Engine/Source/Programs/Online/OnlineTestsCore/Source/Public/Helpers/Sessions/FindSessionsHelper.h
--- a/Engine/Source/Programs/Online/OnlineTestsCore/Source/Public/Helpers/Sessions/FindSessionsHelper.h
+++ b/Engine/Source/Programs/Online/OnlineTestsCore/Source/Public/Helpers/Sessions/FindSessionsHelper.h
@@ -28,7 +28,7 @@ struct FFindSessionsHelper : public FAsyncTestStep
 
 	virtual ~FFindSessionsHelper() = default;
 
-	virtual void Run(FAsyncStepResult Promise, SubsystemType Services) override
+	virtual void Run(FAsyncStepResult Promise, const IOnlineServicesPtr& Services) override
 	{	
 		SessionsInterface = Services->GetSessionsInterface();
 		REQUIRE(SessionsInterface);
diff --git a/Engine/Source/Programs/Online/OnlineTestsCore/Source/Public/Helpers/Sessions/JoinSessionHelper.h b/Engine/Source/Programs/Online/OnlineTestsCore/Source/Public/Helpers/Sessions/JoinSessionHelper.h
--- a/Engine/Source/Programs/Online/OnlineTestsCore/Source/Public/Helpers/Sessions/JoinSessionHelper.h
+++ b/Engine/Source/Programs/Online/OnlineTestsCore/Source/Public/Helpers/Sessions/JoinSessionHelper.h
@@ -25,7 +25,7 @@ struct FJoinSessionHelper : public FAsyncTestStep
 
 	virtual ~FJoinSessionHelper() = default;
 
-	virtual void Run(FAsyncStepResult Promise, SubsystemType Services) override
+	virtual void Run(FAsyncStepResult Promise, const IOnlineServicesPtr& Services) override
 	{
 		SessionsInterface = Services->GetSessionsInterface();
 		REQUIRE(SessionsInterface);
diff --git a/Engine/Source/Programs/Online/OnlineTestsCore/Source/Public/Helpers/Sessions/LeaveSessionHelper.h b/Engine/Source/Programs/Online/OnlineTestsCore/Source/Public/Helpers/Sessions/LeaveSessionHelper.h
--- a/Engine/Source/Programs/Online/OnlineTestsCore/Source/Public/Helpers/Sessions/LeaveSessionHelper.h
+++ b/Engine/Source/Programs/Online/OnlineTestsCore/Source/Public/Helpers/Sessions/LeaveSessionHelper.h
@@ -26,7 +26,7 @@ struct FLeaveSessionHelper : public FAsyncTestStep
 
 	virtual ~FLeaveSessionHelper() = default;
 
-	virtual void Run(FAsyncStepResult Promise, SubsystemType Services) override
+	virtual void Run(FAsyncStepResult Promise, const IOnlineServicesPtr& Services) override
 	{
 		SessionsInterface = Services->GetSessionsInterface();
 		REQUIRE(SessionsInterface);
diff --git a/Engine/Source/Programs/Online/OnlineTestsCore/Source/Public/Helpers/Sessions/SendRejectSessionInviteHelper.h b/Engine/Source/Programs/Online/OnlineTestsCore/Source/Public/Helpers/Sessions/SendRejectSessionInviteHelper.h
--- a/Engine/Source/Programs/Online/OnlineTestsCore/Source/Public/Helpers/Sessions/SendRejectSessionInviteHelper.h
+++ b/Engine/Source/Programs/Online/OnlineTestsCore/Source/Public/Helpers/Sessions/SendRejectSessionInviteHelper.h
@@ -30,7 +30,7 @@ struct FSendSessionInviteHelper : public FAsyncTestStep
 		SessionsInviteReceivedHandle.Unbind();
 	}
 
-	virtual void Run(FAsyncStepResult Promise, SubsystemType Services) override
+	virtual void Run(FAsyncStepResult Promise, const IOnlineServicesPtr& Services) override
 	{
 		SessionsInterface = Services->GetSessionsInterface();
 		REQUIRE(SessionsInterface);
@@ -87,7 +87,7 @@ struct FRejectSessionInviteHelper : public FAsyncTestStep
 
 	virtual ~FRejectSessionInviteHelper() = default;
 
-	virtual void Run(FAsyncStepResult Promise, SubsystemType Services) override
+	virtual void Run(FAsyncStepResult Promise, const IOnlineServicesPtr& Services) override
 	{
 		SessionsInterface = Services->GetSessionsInterface();
 		REQUIRE(SessionsInterface);
diff --git a/Engine/Source/Programs/Online/OnlineTestsCore/Source/Public/Helpers/Sessions/StartMatchmakingSessionHelper.h b/Engine/Source/Programs/Online/OnlineTestsCore/Source/Public/Helpers/Sessions/StartMatchmakingSessionHelper.h
--- a/Engine/Source/Programs/Online/OnlineTestsCore/Source/Public/Helpers/Sessions/StartMatchmakingSessionHelper.h
+++ b/Engine/Source/Programs/Online/OnlineTestsCore/Source/Public/Helpers/Sessions/StartMatchmakingSessionHelper.h
@@ -26,7 +26,7 @@ struct FStartMatchmakingHelper : public FAsyncTestStep
 
 	virtual ~FStartMatchmakingHelper() = default;
 
-	virtual void Run(FAsyncStepResult Promise, SubsystemType Services) override
+	virtual void Run(FAsyncStepResult Promise, const IOnlineServicesPtr& Services) override
 	{
 		SessionsInterface = Services->GetSessionsInterface();
 		REQUIRE(SessionsInterface);
diff --git a/Engine/Source/Programs/Online/OnlineTestsCore/Source/Public/Helpers/Sessions/UpdateSessionSettingsHelper.h b/Engine/Source/Programs/Online/OnlineTestsCore/Source/Public/Helpers/Sessions/UpdateSessionSettingsHelper.h
--- a/Engine/Source/Programs/Online/OnlineTestsCore/Source/Public/Helpers/Sessions/UpdateSessionSettingsHelper.h
+++ b/Engine/Source/Programs/Online/OnlineTestsCore/Source/Public/Helpers/Sessions/UpdateSessionSettingsHelper.h
@@ -26,7 +26,7 @@ struct FUpdateSessionSettingsHelper : public FAsyncTestStep
 
 	virtual ~FUpdateSessionSettingsHelper() = default;
 
-	virtual void Run(FAsyncStepResult Promise, SubsystemType Services) override
+	virtual void Run(FAsyncStepResult Promise, const IOnlineServicesPtr& Services) override
 	{
 		SessionsInterface = Services->GetSessionsInterface();
 		REQUIRE(SessionsInterface);
diff --git a/Engine/Source/Programs/Online/OnlineTestsCore/Source/Public/Helpers/Social/AcceptFriendInviteHelper.h b/Engine/Source/Programs/Online/OnlineTestsCore/Source/Public/Helpers/Social/AcceptFriendInviteHelper.h
--- a/Engine/Source/Programs/Online/OnlineTestsCore/Source/Public/Helpers/Social/AcceptFriendInviteHelper.h
+++ b/Engine/Source/Programs/Online/OnlineTestsCore/Source/Public/Helpers/Social/AcceptFriendInviteHelper.h
@@ -26,7 +26,7 @@ struct FAcceptFriendInviteHelper : public FAsyncTestStep
 
 	virtual ~FAcceptFriendInviteHelper() = default;
 
-	virtual void Run(FAsyncStepResult Promise, SubsystemType Services) override
+	virtual void Run(FAsyncStepResult Promise, const IOnlineServicesPtr& Services) override
 	{
 		SocialInterface = Services->GetSocialInterface();
 		REQUIRE(SocialInterface);
diff --git a/Engine/Source/Programs/Online/OnlineTestsCore/Source/Public/Helpers/Social/BlockUserHelper.h b/Engine/Source/Programs/Online/OnlineTestsCore/Source/Public/Helpers/Social/BlockUserHelper.h
--- a/Engine/Source/Programs/Online/OnlineTestsCore/Source/Public/Helpers/Social/BlockUserHelper.h
+++ b/Engine/Source/Programs/Online/OnlineTestsCore/Source/Public/Helpers/Social/BlockUserHelper.h
@@ -26,7 +26,7 @@ struct FBlockUserHelper : public FAsyncTestStep
 
 	virtual ~FBlockUserHelper() = default;
 
-	virtual void Run(FAsyncStepResult Promise, SubsystemType Services) override
+	virtual void Run(FAsyncStepResult Promise, const IOnlineServicesPtr& Services) override
 	{
 		SocialInterface = Services->GetSocialInterface();
 		REQUIRE(SocialInterface);
diff --git a/Engine/Source/Programs/Online/OnlineTestsCore/Source/Public/Helpers/Social/GetBlockedUsersHelper.h b/Engine/Source/Programs/Online/OnlineTestsCore/Source/Public/Helpers/Social/GetBlockedUsersHelper.h
--- a/Engine/Source/Programs/Online/OnlineTestsCore/Source/Public/Helpers/Social/GetBlockedUsersHelper.h
+++ b/Engine/Source/Programs/Online/OnlineTestsCore/Source/Public/Helpers/Social/GetBlockedUsersHelper.h
@@ -27,7 +27,7 @@ struct FGetBlockedUsersHelper : public FAsyncTestStep
 
 	virtual ~FGetBlockedUsersHelper() = default;
 
-	virtual void Run(FAsyncStepResult Promise, SubsystemType Services) override
+	virtual void Run(FAsyncStepResult Promise, const IOnlineServicesPtr& Services) override
 	{
 		SocialInterface = Services->GetSocialInterface();
 		REQUIRE(SocialInterface);
diff --git a/Engine/Source/Programs/Online/OnlineTestsCore/Source/Public/Helpers/Social/GetFriendsHelper.h b/Engine/Source/Programs/Online/OnlineTestsCore/Source/Public/Helpers/Social/GetFriendsHelper.h
--- a/Engine/Source/Programs/Online/OnlineTestsCore/Source/Public/Helpers/Social/GetFriendsHelper.h
+++ b/Engine/Source/Programs/Online/OnlineTestsCore/Source/Public/Helpers/Social/GetFriendsHelper.h
@@ -27,7 +27,7 @@ struct FGetFriendsHelper : public FAsyncTestStep
 
 	virtual ~FGetFriendsHelper() = default;
 
-	virtual void Run(FAsyncStepResult Promise, SubsystemType Services) override
+	virtual void Run(FAsyncStepResult Promise, const IOnlineServicesPtr& Services) override
 	{
 		SocialInterface = Services->GetSocialInterface();
 		REQUIRE(SocialInterface);
diff --git a/Engine/Source/Programs/Online/OnlineTestsCore/Source/Public/Helpers/Social/QueryBlockedUsersHelper.h b/Engine/Source/Programs/Online/OnlineTestsCore/Source/Public/Helpers/Social/QueryBlockedUsersHelper.h
--- a/Engine/Source/Programs/Online/OnlineTestsCore/Source/Public/Helpers/Social/QueryBlockedUsersHelper.h
+++ b/Engine/Source/Programs/Online/OnlineTestsCore/Source/Public/Helpers/Social/QueryBlockedUsersHelper.h
@@ -26,7 +26,7 @@ struct FQueryBlockedUsersHelper : public FAsyncTestStep
 
 	virtual ~FQueryBlockedUsersHelper() = default;
 
-	virtual void Run(FAsyncStepResult Promise, SubsystemType Services) override
+	virtual void Run(FAsyncStepResult Promise, const IOnlineServicesPtr& Services) override
 	{
 		SocialInterface = Services->GetSocialInterface();
 		REQUIRE(SocialInterface);
diff --git a/Engine/Source/Programs/Online/OnlineTestsCore/Source/Public/Helpers/Social/QueryFriendsHelper.h b/Engine/Source/Programs/Online/OnlineTestsCore/Source/Public/Helpers/Social/QueryFriendsHelper.h
--- a/Engine/Source/Programs/Online/OnlineTestsCore/Source/Public/Helpers/Social/QueryFriendsHelper.h
+++ b/Engine/Source/Programs/Online/OnlineTestsCore/Source/Public/Helpers/Social/QueryFriendsHelper.h
@@ -26,7 +26,7 @@ struct FQueryFriendsHelper : public FAsyncTestStep
 
 	virtual ~FQueryFriendsHelper() = default;
 
-	virtual void Run(FAsyncStepResult Promise, SubsystemType Services) override
+	virtual void Run(FAsyncStepResult Promise, const IOnlineServicesPtr& Services) override
 	{
 		SocialInterface = Services->GetSocialInterface();
 		REQUIRE(SocialInterface);
diff --git a/Engine/Source/Programs/Online/OnlineTestsCore/Source/Public/Helpers/Social/RejectFriendInviteHelper.h b/Engine/Source/Programs/Online/OnlineTestsCore/Source/Public/Helpers/Social/RejectFriendInviteHelper.h
--- a/Engine/Source/Programs/Online/OnlineTestsCore/Source/Public/Helpers/Social/RejectFriendInviteHelper.h
+++ b/Engine/Source/Programs/Online/OnlineTestsCore/Source/Public/Helpers/Social/RejectFriendInviteHelper.h
@@ -26,7 +26,7 @@ struct FRejectFriendInviteHelper : public FAsyncTestStep
 
 	virtual ~FRejectFriendInviteHelper() = default;
 
-	virtual void Run(FAsyncStepResult Promise, SubsystemType Services) override
+	virtual void Run(FAsyncStepResult Promise, const IOnlineServicesPtr& Services) override
 	{
 		SocialInterface = Services->GetSocialInterface();
 		REQUIRE(SocialInterface);
diff --git a/Engine/Source/Programs/Online/OnlineTestsCore/Source/Public/Helpers/Social/SendFriendInviteHelper.h b/Engine/Source/Programs/Online/OnlineTestsCore/Source/Public/Helpers/Social/SendFriendInviteHelper.h
--- a/Engine/Source/Programs/Online/OnlineTestsCore/Source/Public/Helpers/Social/SendFriendInviteHelper.h
+++ b/Engine/Source/Programs/Online/OnlineTestsCore/Source/Public/Helpers/Social/SendFriendInviteHelper.h
@@ -26,7 +26,7 @@ struct FSendFriendInviteHelper : public FAsyncTestStep
 
 	virtual ~FSendFriendInviteHelper() = default;
 
-	virtual void Run(FAsyncStepResult Promise, SubsystemType Services) override
+	virtual void Run(FAsyncStepResult Promise, const IOnlineServicesPtr& Services) override
 	{
 		SocialInterface = Services->GetSocialInterface();
 		REQUIRE(SocialInterface);
diff --git a/Engine/Source/Programs/Online/OnlineTestsCore/Source/Public/Helpers/Stats/UpdateStatsHelper.h b/Engine/Source/Programs/Online/OnlineTestsCore/Source/Public/Helpers/Stats/UpdateStatsHelper.h
--- a/Engine/Source/Programs/Online/OnlineTestsCore/Source/Public/Helpers/Stats/UpdateStatsHelper.h
+++ b/Engine/Source/Programs/Online/OnlineTestsCore/Source/Public/Helpers/Stats/UpdateStatsHelper.h
@@ -7,7 +7,7 @@
 namespace UE::Online
 {
 
-inline void UpdateStats_Fixture(FAsyncLambdaResult Promise, SubsystemType Services, const FAccountId& AccountId,  const TArray<FUserStats>& UpdateUsersStats)
+inline void UpdateStats_Fixture(FAsyncLambdaResult Promise, const IOnlineServicesPtr& Services, const FAccountId& AccountId,  const TArray<FUserStats>& UpdateUsersStats)
 {
 	Services->GetStatsInterface()->UpdateStats({AccountId,  UpdateUsersStats}).OnComplete([Promise=MoveTemp(Promise)](const TOnlineResult<UE::Online::FUpdateStats>& Result) mutable
 	{
diff --git a/Engine/Source/Programs/Online/OnlineTestsCore/Source/Public/Helpers/TickForTime.h b/Engine/Source/Programs/Online/OnlineTestsCore/Source/Public/Helpers/TickForTime.h
--- a/Engine/Source/Programs/Online/OnlineTestsCore/Source/Public/Helpers/TickForTime.h
+++ b/Engine/Source/Programs/Online/OnlineTestsCore/Source/Public/Helpers/TickForTime.h
@@ -12,7 +12,7 @@ public:
 	{
 	}
  
-	virtual EContinuance Tick(SubsystemType OnlineSubsystem) override
+	virtual EContinuance Tick(const IOnlineServicesPtr& OnlineSubsystem) override
 	{
 		double Now = FPlatformTime::Seconds();
  
diff --git a/Engine/Source/Programs/Online/OnlineTestsCore/Source/Public/Helpers/TitleFile/AsyncTitleFileEnumerateFiles.h b/Engine/Source/Programs/Online/OnlineTestsCore/Source/Public/Helpers/TitleFile/AsyncTitleFileEnumerateFiles.h
--- a/Engine/Source/Programs/Online/OnlineTestsCore/Source/Public/Helpers/TitleFile/AsyncTitleFileEnumerateFiles.h
+++ b/Engine/Source/Programs/Online/OnlineTestsCore/Source/Public/Helpers/TitleFile/AsyncTitleFileEnumerateFiles.h
@@ -25,7 +25,7 @@ public:
 	{
 	}
 
-	virtual void Run(FAsyncStepResult Promise, SubsystemType Services) override
+	virtual void Run(FAsyncStepResult Promise, const IOnlineServicesPtr& Services) override
 	{
 		REQUIRE(Services != nullptr);
 		UE::Online::ITitleFilePtr TitleFile = Services->GetTitleFileInterface();
diff --git a/Engine/Source/Programs/Online/OnlineTestsCore/Source/Public/Helpers/TitleFile/AsyncTitleFileGetEnumeratedFiles.h b/Engine/Source/Programs/Online/OnlineTestsCore/Source/Public/Helpers/TitleFile/AsyncTitleFileGetEnumeratedFiles.h
--- a/Engine/Source/Programs/Online/OnlineTestsCore/Source/Public/Helpers/TitleFile/AsyncTitleFileGetEnumeratedFiles.h
+++ b/Engine/Source/Programs/Online/OnlineTestsCore/Source/Public/Helpers/TitleFile/AsyncTitleFileGetEnumeratedFiles.h
@@ -17,7 +17,7 @@ public:
 	{
 	}
 
-	virtual void Run(FAsyncStepResult Promise, SubsystemType Services) override
+	virtual void Run(FAsyncStepResult Promise, const IOnlineServicesPtr& Services) override
 	{
 		REQUIRE(Services != nullptr);
 		UE::Online::ITitleFilePtr TitleFile = Services->GetTitleFileInterface();
diff --git a/Engine/Source/Programs/Online/OnlineTestsCore/Source/Public/Helpers/TitleFile/AsyncTitleFileReadFile.h b/Engine/Source/Programs/Online/OnlineTestsCore/Source/Public/Helpers/TitleFile/AsyncTitleFileReadFile.h
--- a/Engine/Source/Programs/Online/OnlineTestsCore/Source/Public/Helpers/TitleFile/AsyncTitleFileReadFile.h
+++ b/Engine/Source/Programs/Online/OnlineTestsCore/Source/Public/Helpers/TitleFile/AsyncTitleFileReadFile.h
@@ -18,7 +18,7 @@ public:
 	{
 	}
 
-	virtual void Run(FAsyncStepResult Promise, SubsystemType Services) override
+	virtual void Run(FAsyncStepResult Promise, const IOnlineServicesPtr& Services) override
 	{
 		REQUIRE(Services != nullptr);
 		UE::Online::ITitleFilePtr TitleFile = Services->GetTitleFileInterface();
diff --git a/Engine/Source/Programs/Online/OnlineTestsCore/Source/Public/Helpers/UserFile/AsyncUserFileCopyFile.h b/Engine/Source/Programs/Online/OnlineTestsCore/Source/Public/Helpers/UserFile/AsyncUserFileCopyFile.h
--- a/Engine/Source/Programs/Online/OnlineTestsCore/Source/Public/Helpers/UserFile/AsyncUserFileCopyFile.h
+++ b/Engine/Source/Programs/Online/OnlineTestsCore/Source/Public/Helpers/UserFile/AsyncUserFileCopyFile.h
@@ -25,7 +25,7 @@ public:
 	{
 	}
 
-	virtual void Run(FAsyncStepResult Promise, SubsystemType Services) override
+	virtual void Run(FAsyncStepResult Promise, const IOnlineServicesPtr& Services) override
 	{
 		REQUIRE(Services != nullptr);
 		UE::Online::IUserFilePtr UserFile = Services->GetUserFileInterface();
diff --git a/Engine/Source/Programs/Online/OnlineTestsCore/Source/Public/Helpers/UserFile/AsyncUserFileDeleteFile.h b/Engine/Source/Programs/Online/OnlineTestsCore/Source/Public/Helpers/UserFile/AsyncUserFileDeleteFile.h
--- a/Engine/Source/Programs/Online/OnlineTestsCore/Source/Public/Helpers/UserFile/AsyncUserFileDeleteFile.h
+++ b/Engine/Source/Programs/Online/OnlineTestsCore/Source/Public/Helpers/UserFile/AsyncUserFileDeleteFile.h
@@ -32,7 +32,7 @@ public:
 	{
 	}
 
-	virtual void Run(FAsyncStepResult Promise, SubsystemType Services) override
+	virtual void Run(FAsyncStepResult Promise, const IOnlineServicesPtr& Services) override
 	{
 		REQUIRE(Services != nullptr);
 		UE::Online::IUserFilePtr UserFile = Services->GetUserFileInterface();
diff --git a/Engine/Source/Programs/Online/OnlineTestsCore/Source/Public/Helpers/UserFile/AsyncUserFileEnumerateFiles.h b/Engine/Source/Programs/Online/OnlineTestsCore/Source/Public/Helpers/UserFile/AsyncUserFileEnumerateFiles.h
--- a/Engine/Source/Programs/Online/OnlineTestsCore/Source/Public/Helpers/UserFile/AsyncUserFileEnumerateFiles.h
+++ b/Engine/Source/Programs/Online/OnlineTestsCore/Source/Public/Helpers/UserFile/AsyncUserFileEnumerateFiles.h
@@ -27,7 +27,7 @@ public:
 	{
 	}
 
-	virtual void Run(FAsyncStepResult Promise, SubsystemType Services) override
+	virtual void Run(FAsyncStepResult Promise, const IOnlineServicesPtr& Services) override
 	{
 		REQUIRE(Services != nullptr);
 		UE::Online::IUserFilePtr UserFile = Services->GetUserFileInterface();
diff --git a/Engine/Source/Programs/Online/OnlineTestsCore/Source/Public/Helpers/UserFile/AsyncUserFileGetEnumeratedFiles.h b/Engine/Source/Programs/Online/OnlineTestsCore/Source/Public/Helpers/UserFile/AsyncUserFileGetEnumeratedFiles.h
--- a/Engine/Source/Programs/Online/OnlineTestsCore/Source/Public/Helpers/UserFile/AsyncUserFileGetEnumeratedFiles.h
+++ b/Engine/Source/Programs/Online/OnlineTestsCore/Source/Public/Helpers/UserFile/AsyncUserFileGetEnumeratedFiles.h
@@ -25,7 +25,7 @@ public:
 	{
 	}
 
-	virtual void Run(FAsyncStepResult Promise, SubsystemType Services) override
+	virtual void Run(FAsyncStepResult Promise, const IOnlineServicesPtr& Services) override
 	{
 		REQUIRE(Services != nullptr);
 		UE::Online::IUserFilePtr UserFile = Services->GetUserFileInterface();
diff --git a/Engine/Source/Programs/Online/OnlineTestsCore/Source/Public/Helpers/UserFile/AsyncUserFileReadFile.h b/Engine/Source/Programs/Online/OnlineTestsCore/Source/Public/Helpers/UserFile/AsyncUserFileReadFile.h
--- a/Engine/Source/Programs/Online/OnlineTestsCore/Source/Public/Helpers/UserFile/AsyncUserFileReadFile.h
+++ b/Engine/Source/Programs/Online/OnlineTestsCore/Source/Public/Helpers/UserFile/AsyncUserFileReadFile.h
@@ -25,7 +25,7 @@ public:
 	{
 	}
 
-	virtual void Run(FAsyncStepResult Promise, SubsystemType Services) override
+	virtual void Run(FAsyncStepResult Promise, const IOnlineServicesPtr& Services) override
 	{
 		REQUIRE(Services != nullptr);
 		UE::Online::IUserFilePtr UserFile = Services->GetUserFileInterface();
diff --git a/Engine/Source/Programs/Online/OnlineTestsCore/Source/Public/Helpers/UserFile/AsyncUserFileWriteFile.h b/Engine/Source/Programs/Online/OnlineTestsCore/Source/Public/Helpers/UserFile/AsyncUserFileWriteFile.h
--- a/Engine/Source/Programs/Online/OnlineTestsCore/Source/Public/Helpers/UserFile/AsyncUserFileWriteFile.h
+++ b/Engine/Source/Programs/Online/OnlineTestsCore/Source/Public/Helpers/UserFile/AsyncUserFileWriteFile.h
@@ -25,7 +25,7 @@ public:
 	{
 	}
 
-	virtual void Run(FAsyncStepResult Promise, SubsystemType Services) override
+	virtual void Run(FAsyncStepResult Promise, const IOnlineServicesPtr& Services) override
 	{
 		REQUIRE(Services != nullptr);
 		UE::Online::IUserFilePtr UserFile = Services->GetUserFileInterface();
diff --git a/Engine/Source/Programs/Online/OnlineTestsCore/Source/Public/OnlineCatchHelper.h b/Engine/Source/Programs/Online/OnlineTestsCore/Source/Public/OnlineCatchHelper.h
--- a/Engine/Source/Programs/Online/OnlineTestsCore/Source/Public/OnlineCatchHelper.h
+++ b/Engine/Source/Programs/Online/OnlineTestsCore/Source/Public/OnlineCatchHelper.h
@@ -49,7 +49,7 @@ public:
 	typedef UE::Online::FAccountId FAccountId;
 public:
 
-	/* Attempts to assgin OutAccountId to what LocalUserId is logged in as. */
+	/* Attempts to assign OutAccountId to what LocalUserId is logged in as. */
 	void AssignLoginUsers(int32 LocalUserId, FAccountId& OutAccountId) const;
 
 #if ONLINETESTS_USEEXTERNAUTH
@@ -73,11 +73,11 @@ protected:
 	/* Returns the name of the service we're currently testing */
 	FString GetService() const;
 	UE::Online::EOnlineServices GetServiceType() const;
-	SubsystemType GetSubsystem() const;
+	IOnlineServicesPtr GetSubsystem() const;
 
 // If you define a custom credentials path you must define these in your own Auth.cpp file
 #if ONLINETESTS_USEEXTERNAUTH
-	// This will provide a paramters to pass to the OSS Auth interface.
+	// This will provide a parameters to pass to the OSS Auth interface.
 	TArray<UE::Online::FAuthLogin::Params> CustomCredentials(int32 LocalUserNum, int32 NumUsers) const;
 
 	// When called should "reset" accounts to default state for Stats, Achivements, Lobbies and other
diff --git a/Engine/Source/Programs/Online/OnlineTestsCore/Source/Public/TestDriver.h b/Engine/Source/Programs/Online/OnlineTestsCore/Source/Public/TestDriver.h
--- a/Engine/Source/Programs/Online/OnlineTestsCore/Source/Public/TestDriver.h
+++ b/Engine/Source/Programs/Online/OnlineTestsCore/Source/Public/TestDriver.h
@@ -12,7 +12,7 @@
 #define UE_BUILD_DEBUG 0
 #endif
 
-typedef TSharedPtr<UE::Online::IOnlineServices> SubsystemType;
+typedef TSharedPtr<UE::Online::IOnlineServices> IOnlineServicesPtr;
 typedef TSharedPtr<TPromise<bool>> FAsyncLambdaResult;
 
 DEFINE_LOG_CATEGORY_STATIC(LogOnlineTests, Log, Log);
@@ -74,7 +74,7 @@ public:
 		 */
 		FTimespan AbsoluteMaximumTick = FTimespan::FromMilliseconds(75);
 		/**
-		 * The minimum number of ticks per test before the average tick is evaluated against ExpectedAvergateTick.
+		 * The minimum number of ticks per test before the average tick is evaluated against ExpectedAverageTick.
 		 * The default is changed for debug.
 		 */
 		uint32_t MinimumTickCount = 10;
@@ -112,35 +112,35 @@ public:
 		virtual bool RequiresDeletePostRelease() const { return false; }
 		// Called only if RequestDeletePostRelease is true and the FStep is being stored.
 		virtual void OnPreRelease() {}
-		virtual EContinuance Tick(SubsystemType Subsystem) = 0;
+		virtual EContinuance Tick(const IOnlineServicesPtr& Services) = 0;
 	};
 	using FStepPtr = TUniquePtr<FStep>;
 
 	class FLambdaStep : public FStep
 	{
 	public:
-		FLambdaStep(TUniqueFunction<void(SubsystemType)>&& InLambda)
+		FLambdaStep(TUniqueFunction<void(IOnlineServicesPtr)>&& InLambda)
 			: Lambda(MoveTemp(InLambda))
 		{}
 
-		virtual EContinuance Tick(SubsystemType Subsystem) override
+		virtual EContinuance Tick(const IOnlineServicesPtr& Services) override
 		{
-			Lambda(Subsystem);
+			Lambda(Services);
 			return EContinuance::Done;
 		}
 
 	private:
-		TUniqueFunction<void(SubsystemType)> Lambda;
+		TUniqueFunction<void(IOnlineServicesPtr)> Lambda;
 	};
 
 	class FAsyncLambdaStep : public FStep
 	{
 	public:
-		FAsyncLambdaStep(TUniqueFunction<void(FAsyncLambdaResult, SubsystemType)>&& InLambda)
+		FAsyncLambdaStep(TUniqueFunction<void(FAsyncLambdaResult, IOnlineServicesPtr)>&& InLambda)
 			: Lambda(MoveTemp(InLambda))
 		{}
 
-		virtual EContinuance Tick(SubsystemType Subsystem) override
+		virtual EContinuance Tick(const IOnlineServicesPtr& Services) override
 		{
 			if (!bRan)
 			{
@@ -151,7 +151,7 @@ public:
 					REQUIRE(Result);
 					bComplete = true;
 				});
-				Lambda(Result, Subsystem);
+				Lambda(Result, Services);
 			}
 
 			if (bComplete == true)
@@ -165,7 +165,7 @@ public:
 		}
 
 	private:
-		TUniqueFunction<void(FAsyncLambdaResult, SubsystemType)> Lambda;
+		TUniqueFunction<void(FAsyncLambdaResult, IOnlineServicesPtr)> Lambda;
 		bool bRan = false;
 		TAtomic<bool> bComplete = false;
 	};
@@ -199,7 +199,7 @@ public:
 	FTestPipeline&& EmplaceStep(TArguments&&... Args)
 	{
 		static_assert(std::is_base_of_v<FStep, T>, "Step type must be derived from FTestPipeline::FStep");
-		uint32_t NewIndex = (uint32_t)TestSteps.Num();
+		uint32_t NewIndex = static_cast<uint32_t>(TestSteps.Num());
 		TestSteps.Emplace(MakeUnique<T>(Forward<TArguments>(Args)...));
 		TestSteps.Last()->Index = NewIndex;
 
@@ -207,12 +207,12 @@ public:
 	}
 
 	/** Adds a custom test step. This function will run and resolve immediately */
-	FTestPipeline&& EmplaceLambda(TUniqueFunction<void(SubsystemType)> Lambda)
+	FTestPipeline&& EmplaceLambda(TUniqueFunction<void(IOnlineServicesPtr)> Lambda)
 	{
 		return EmplaceStep<FLambdaStep>(MoveTemp(Lambda));
 	}
 
-	FTestPipeline&& EmplaceAsyncLambda(TUniqueFunction<void(FAsyncLambdaResult, SubsystemType)> Lambda)
+	FTestPipeline&& EmplaceAsyncLambda(TUniqueFunction<void(FAsyncLambdaResult, IOnlineServicesPtr)> Lambda)
 	{
 		return EmplaceStep<FAsyncLambdaStep>(MoveTemp(Lambda));
 	}
@@ -220,7 +220,7 @@ public:
 	/** Generates a string suitable for INFO which will help identify where and when a specific failure has occurred. */
 	FString InfoString() const;
 
-	void operator()(SubsystemType Subsystem);
+	void operator()(const IOnlineServicesPtr& Services);
 
 	void OnPreRelease();
 
@@ -259,7 +259,7 @@ private:
 	TArray<FStepPtr> TimedoutSteps;
 	TArray<FStepPtr> DeletePostReleaseSteps;
 	FTimespan PipelineTimeout;
-	double PipelineStartTime;
+	double PipelineStartTime = 0.0;
 	/** Current excess tick config. */
 	FEvaluateTickConfig EvaluateTickConfig;
 	/** Sum of time to process all calls to Subsystem Tick. */
@@ -273,7 +273,7 @@ private:
 class FTestDriver
 {
 public:
-	using FSubsystemInstanceMap = TMap<SubsystemType, FTestPipeline>;
+	using FSubsystemInstanceMap = TMap<IOnlineServicesPtr, FTestPipeline>;
 
 	~FTestDriver();
 
@@ -299,7 +299,7 @@ public:
 
 	bool AddPipeline(FTestPipeline&& Pipeline, const FPipelineTestContext& TestContext = FPipelineTestContext(TEXT("NULL"), UE::Online::EOnlineServices::Null));
 
-	void MarkComplete(SubsystemType Key);
+	void MarkComplete(IOnlineServicesPtr Key);
 
 	void RunToCompletion();
 
@@ -314,7 +314,7 @@ private:
 	void FlushCompleted();
 
 	FSubsystemInstanceMap SubsystemInstances;
-	TSet<SubsystemType> CompletedInstances;
+	TSet<IOnlineServicesPtr> CompletedInstances;
 	bool bDidTimeout = false;
 	double LastTickTime = 0.0;
 };
