OSMcp basic party tests

diff --git a/Engine/Source/Programs/Online/OSSTestsCore/OnlineSubsystemCatchHelper.cpp b/Engine/Source/Programs/Online/OSSTestsCore/OnlineSubsystemCatchHelper.cpp
--- a/Engine/Source/Programs/Online/OSSTestsCore/OnlineSubsystemCatchHelper.cpp
+++ b/Engine/Source/Programs/Online/OSSTestsCore/OnlineSubsystemCatchHelper.cpp
@@ -193,7 +193,7 @@ FTestPipeline& FOnlineSubsystemTestBaseFixture::GetPipeline() const
 	return GetLoginPipeline(0);
 }
 
-void FOnlineSubsystemTestBaseFixture::RunToCompletion(bool bWaitBeforeLogout, const FTimespan TimeToWaitMilliseconds) const
+void FOnlineSubsystemTestBaseFixture::RunToCompletion(bool bWaitBeforeLogout, const FTimespan TimeToWaitMilliseconds, const FString SubsystemInstanceName) const
 {
 	bool bUseAutoLogin = false;
 	bool bUseImplicitLogin = false;
@@ -225,7 +225,7 @@ void FOnlineSubsystemTestBaseFixture::RunToCompletion(bool bWaitBeforeLogout, co
 	}
 	
 	FName SubsystemName = FName(GetSubsystem());
-	FPipelineTestContext TestContext = FPipelineTestContext(SubsystemName);
+	FPipelineTestContext TestContext = FPipelineTestContext(Subsystem, SubsystemInstanceName);
 	CHECK(Driver.AddPipeline(MoveTemp(Pipeline), TestContext));
 	REQUIRE(IOnlineSubsystem::IsEnabled(SubsystemName));
 	Driver.RunToCompletion();
diff --git a/Engine/Source/Programs/Online/OSSTestsCore/OnlineSubsystemCatchHelper.h b/Engine/Source/Programs/Online/OSSTestsCore/OnlineSubsystemCatchHelper.h
--- a/Engine/Source/Programs/Online/OSSTestsCore/OnlineSubsystemCatchHelper.h
+++ b/Engine/Source/Programs/Online/OSSTestsCore/OnlineSubsystemCatchHelper.h
@@ -155,7 +155,7 @@ protected:
 	/* Returns the ini login category name for the configured service */
 	FString GetLoginCredentialCategory() const;
 
-	void RunToCompletion(bool bWaitBeforeLogout = false, const FTimespan TimeToWaitMilliseconds = FTimespan::FromMilliseconds(1000)) const;
+	void RunToCompletion(bool bWaitBeforeLogout = false, const FTimespan TimeToWaitMilliseconds = FTimespan::FromMilliseconds(1000), const FString SubsystemInstanceName = TEXT("")) const;
 
 	/* ITestInvoker */
 	virtual void invoke() const override = 0;
diff --git a/Engine/Source/Programs/Online/OSSTestsCore/TestDriver.cpp b/Engine/Source/Programs/Online/OSSTestsCore/TestDriver.cpp
--- a/Engine/Source/Programs/Online/OSSTestsCore/TestDriver.cpp
+++ b/Engine/Source/Programs/Online/OSSTestsCore/TestDriver.cpp
@@ -156,7 +156,10 @@ FTestDriver::~FTestDriver()
 
 bool FTestDriver::AddPipeline(FTestPipeline&& Pipeline, const FPipelineTestContext& TestContext)
 {
-	IOnlineSubsystem* OnlineSubsystem = IOnlineSubsystem::Get(TestContext.InitOptions.SubsystemType);
+	FString CompleteName = TestContext.InitOptions.InstanceName.IsEmpty()
+		? TestContext.InitOptions.SubsystemType
+		: TestContext.InitOptions.SubsystemType + TEXT(":") + TestContext.InitOptions.InstanceName;
+	IOnlineSubsystem* OnlineSubsystem = IOnlineSubsystem::Get(FName(CompleteName));
 	if (OnlineSubsystem == nullptr)
 	{
 		return false;
diff --git a/Engine/Source/Programs/Online/OSSTestsCore/TestDriver.h b/Engine/Source/Programs/Online/OSSTestsCore/TestDriver.h
--- a/Engine/Source/Programs/Online/OSSTestsCore/TestDriver.h
+++ b/Engine/Source/Programs/Online/OSSTestsCore/TestDriver.h
@@ -21,7 +21,8 @@ class FTestDriver;
 
 struct FPipelineTestContextInitOptions
 {
-	FName SubsystemType;
+	FString SubsystemType;
+	FString InstanceName;
 };
 
 class FPipelineTestContext
@@ -30,9 +31,10 @@ class FPipelineTestContext
 public:
 	FPipelineTestContextInitOptions InitOptions;
 
-	FPipelineTestContext(const FName& SubsystemType)
+	FPipelineTestContext(const FString& SubsystemType, const FString& InstanceName = TEXT(""))
 	{
 		InitOptions.SubsystemType = SubsystemType;
+		InitOptions.InstanceName = InstanceName;
 	}
 
 };
@@ -245,7 +247,7 @@ public:
 		return FTestPipeline(*this, Forward<TArgs>(Args)...);
 	}
 
-	bool AddPipeline(FTestPipeline&& Pipeline, const FPipelineTestContext& TestContext = FPipelineTestContext(NULL_SUBSYSTEM));
+	bool AddPipeline(FTestPipeline&& Pipeline, const FPipelineTestContext& TestContext = FPipelineTestContext(NULL_SUBSYSTEM.ToString()));
 
 	void MarkComplete(IOnlineSubsystem* Key);
 
diff --git a/Engine/Source/Programs/Online/OSSTestsCore/Tests/Identity/IdentityAutoLoginTests.cpp b/Engine/Source/Programs/Online/OSSTestsCore/Tests/Identity/IdentityAutoLoginTests.cpp
--- a/Engine/Source/Programs/Online/OSSTestsCore/Tests/Identity/IdentityAutoLoginTests.cpp
+++ b/Engine/Source/Programs/Online/OSSTestsCore/Tests/Identity/IdentityAutoLoginTests.cpp
@@ -35,7 +35,7 @@ IDENTITY_LEGACY_AUTOLOGIN_EOS_TEST("Verify calling Identity AutoLogin with valid
 		.EmplaceStep<FIdentityAutoLoginStep>(LocalUserNum)
 		.EmplaceStep<FIdentityLogoutStep>(LocalUserNum);
 
-	FPipelineTestContext TestContext = FPipelineTestContext(FName(GetSubsystem()));
+	FPipelineTestContext TestContext = FPipelineTestContext(GetSubsystem());
 	REQUIRE(LocalDriver.AddPipeline(MoveTemp(LocalPipeline), TestContext));
 	LocalDriver.RunToCompletion();
 }
@@ -58,7 +58,7 @@ IDENTITY_TEST_CASE("Verify calling Identity AutoLogin with valid inputs returns
 		.EmplaceStep<FIdentityAutoLoginStep>(LocalUserNum)
 		.EmplaceStep<FIdentityLogoutStep>(LocalUserNum);
 
-	FPipelineTestContext TestContext = FPipelineTestContext(FName(GetSubsystem()));
+	FPipelineTestContext TestContext = FPipelineTestContext(GetSubsystem());
 	REQUIRE(LocalDriver.AddPipeline(MoveTemp(LocalPipeline), TestContext));
 	LocalDriver.RunToCompletion();
 }
diff --git a/Engine/Source/Programs/Online/OSSTestsCore/Tests/Identity/IdentityGetLoginStatusTests.cpp b/Engine/Source/Programs/Online/OSSTestsCore/Tests/Identity/IdentityGetLoginStatusTests.cpp
--- a/Engine/Source/Programs/Online/OSSTestsCore/Tests/Identity/IdentityGetLoginStatusTests.cpp
+++ b/Engine/Source/Programs/Online/OSSTestsCore/Tests/Identity/IdentityGetLoginStatusTests.cpp
@@ -33,7 +33,7 @@ IDENTITY_TEST_CASE("Verify calling Identity GetLoginStatus by UserId with valid
 		.EmplaceStep<FIdentityLogoutStep>(LocalUserNum)
 		.EmplaceStep<FIdentityGetLoginStatusByUserIdStep>(&LocalUserId, ExpectedLogoutStatus);
 
-	REQUIRE(LocalDriver.AddPipeline(MoveTemp(LocalPipeline), FName(GetSubsystem())));
+	REQUIRE(LocalDriver.AddPipeline(MoveTemp(LocalPipeline), GetSubsystem()));
 	LocalDriver.RunToCompletion();
 }
 
@@ -53,6 +53,6 @@ IDENTITY_TEST_CASE("Verify calling Identity GetLoginStatus by LocalUserNum with
 		.EmplaceStep<FIdentityLogoutStep>(LocalUserNum)
 		.EmplaceStep<FIdentityGetLoginStatusByLocalUserNumStep>(LocalUserNum, ExpectedLogoutStatus);
 
-	REQUIRE(LocalDriver.AddPipeline(MoveTemp(LocalPipeline), FName(GetSubsystem())));
+	REQUIRE(LocalDriver.AddPipeline(MoveTemp(LocalPipeline), GetSubsystem()));
 	LocalDriver.RunToCompletion();
 }
