OSMcp test: send invitations to party then log out, must restore invitations upon log in otherwise they are gone by default

#rnx

diff --git a/Engine/Source/Programs/Online/OSSTestsCore/OnlineSubsystemCatchHelper.cpp b/Engine/Source/Programs/Online/OSSTestsCore/OnlineSubsystemCatchHelper.cpp
--- a/Engine/Source/Programs/Online/OSSTestsCore/OnlineSubsystemCatchHelper.cpp
+++ b/Engine/Source/Programs/Online/OSSTestsCore/OnlineSubsystemCatchHelper.cpp
@@ -151,9 +151,12 @@ FString FOnlineSubsystemTestBaseFixture::GetLoginCredentialCategory() const
 	return FString::Printf(TEXT("LoginCredentials %s"), *Subsystem);
 }
 
-FTestPipeline& FOnlineSubsystemTestBaseFixture::GetLoginPipeline(uint32 NumUsersToLogin) const
+FTestPipeline& FOnlineSubsystemTestBaseFixture::GetLoginPipeline(uint32 NumUsersToLogin, bool MultiLogin) const
 {
-	REQUIRE(NumLocalUsers == -1); // Don't call GetLoginPipeline more than once per test
+	if (!MultiLogin)
+	{
+		REQUIRE(NumLocalUsers == -1); // Don't call GetLoginPipeline more than once per test
+	}
 	NumLocalUsers = NumUsersToLogin;
 	NumUsersToLogout = NumUsersToLogin;
 
diff --git a/Engine/Source/Programs/Online/OSSTestsCore/OnlineSubsystemCatchHelper.h b/Engine/Source/Programs/Online/OSSTestsCore/OnlineSubsystemCatchHelper.h
--- a/Engine/Source/Programs/Online/OSSTestsCore/OnlineSubsystemCatchHelper.h
+++ b/Engine/Source/Programs/Online/OSSTestsCore/OnlineSubsystemCatchHelper.h
@@ -149,7 +149,7 @@ protected:
 
 	TArray<FOnlineAccountCredentials> GetCredentials(int32 TestAccountIndex, int32 NumUsers) const;
 
-	FTestPipeline& GetLoginPipeline(uint32 NumUsersToLogin = 1) const;
+	FTestPipeline& GetLoginPipeline(uint32 NumUsersToLogin = 1, bool MultiLogin = false) const;
 	FTestPipeline& GetPipeline() const;
 
 	/* Returns the ini login category name for the configured service */
diff --git a/Engine/Source/Programs/Online/OSSTestsCore/TestDriver.cpp b/Engine/Source/Programs/Online/OSSTestsCore/TestDriver.cpp
--- a/Engine/Source/Programs/Online/OSSTestsCore/TestDriver.cpp
+++ b/Engine/Source/Programs/Online/OSSTestsCore/TestDriver.cpp
@@ -58,43 +58,73 @@ void FTestPipeline::operator()(IOnlineSubsystem* Subsystem)
 	}
 
 	FStepPtr& CurrentStep = TestSteps.HeapTop();
-	switch (CurrentStep->Tick(Subsystem))
+
+	if (CurrentStep->StepWaitsForCompletion())
+	{
+		if (CheckStepTimedOut(CurrentStep, CurrentTime, Subsystem))
+		{
+			return;
+		}
+
+		if (CurrentStep->Tick(Subsystem) == FStep::EContinuance::Done)
+		{
+			CompletedSteps.Emplace(MoveTemp(CurrentStep));
+			TestSteps.RemoveAt(0);
+		}
+		else
+		{
+			return; // Wait until state is Done
+		}
+	}
+	else
 	{
+		switch (CurrentStep->Tick(Subsystem))
+		{
 
-	case FStep::EContinuance::Done:
-		// Retain the completed steps for those which have long term notify handlers.
-		CompletedSteps.Emplace(MoveTemp(CurrentStep));
-		TestSteps.RemoveAt(0);
-		break;
+		case FStep::EContinuance::Done:
+			// Retain the completed steps for those which have long term notify handlers.
+			CompletedSteps.Emplace(MoveTemp(CurrentStep));
+			TestSteps.RemoveAt(0);
+			break;
 
-	default:
+		default:
 
-		if (FTimespan::FromSeconds(CurrentTime - PipelineStartTime).GetTicks() >= PipelineTimeout.GetTicks())
-		{
-			Driver.SetDriverTimedOut(true);
-			if (CurrentStep->IsOptional())
+			if (CheckStepTimedOut(CurrentStep, CurrentTime, Subsystem))
 			{
-				// Mark only this step as timed out.
-				// Retain the timed-out steps since the callback has yet to be handled.
-				TimedoutSteps.Emplace(MoveTemp(CurrentStep));
-				TestSteps.RemoveAt(0);
+				return;
 			}
-			else
+			break;
+		}
+	}
+}
+
+bool FTestPipeline::CheckStepTimedOut(FStepPtr& CurrentStep, double CurrentTime, IOnlineSubsystem* Subsystem)
+{
+	if (FTimespan::FromSeconds(CurrentTime - PipelineStartTime).GetTicks() >= PipelineTimeout.GetTicks())
+	{
+		Driver.SetDriverTimedOut(true);
+		if (CurrentStep->IsOptional())
+		{
+			// Mark only this step as timed out.
+			// Retain the timed-out steps since the callback has yet to be handled.
+			TimedoutSteps.Emplace(MoveTemp(CurrentStep));
+			TestSteps.RemoveAt(0);
+		}
+		else
+		{
+			// Mark all remaining steps as timed out.
+			while (!TestSteps.IsEmpty())
 			{
-				// Mark all remaining steps as timed out.
-				while (!TestSteps.IsEmpty())
-				{
-					FTestPipeline::FStepPtr& ExistingStep = TestSteps.HeapTop();
-					// Retain the timed-out steps since the callback may not have been handled.
-					TimedoutSteps.Emplace(MoveTemp(ExistingStep));
-					TestSteps.RemoveAt(0);
-				}
-				Driver.MarkComplete(Subsystem);
-				return;
+				FTestPipeline::FStepPtr& ExistingStep = TestSteps.HeapTop();
+				// Retain the timed-out steps since the callback may not have been handled.
+				TimedoutSteps.Emplace(MoveTemp(ExistingStep));
+				TestSteps.RemoveAt(0);
 			}
+			Driver.MarkComplete(Subsystem);
+			return true;
 		}
-		break;
 	}
+	return false;
 }
 
 void FTestPipeline::OnPreRelease()
diff --git a/Engine/Source/Programs/Online/OSSTestsCore/TestDriver.h b/Engine/Source/Programs/Online/OSSTestsCore/TestDriver.h
--- a/Engine/Source/Programs/Online/OSSTestsCore/TestDriver.h
+++ b/Engine/Source/Programs/Online/OSSTestsCore/TestDriver.h
@@ -102,10 +102,14 @@ public:
 		enum class EContinuance { ContinueStepping, Done };
 		virtual ~FStep() {}
 		virtual bool IsOptional() const { return false; }
+		// Wait to complete before next steps can Tick
+		bool StepWaitsForCompletion() const { return WaitForCompletion; }
 		virtual bool RequiresDeletePostRelease() const { return false; }
 		// Called only if RequestDeletePostRelease is true and the FStep is being stored.
 		virtual void OnPreRelease() {}
 		virtual EContinuance Tick(IOnlineSubsystem* Subsystem) = 0;
+	protected:
+		bool WaitForCompletion = false; // wait before next steps can Tick
 	};
 	using FStepPtr = TUniquePtr<FStep>;
 
@@ -201,6 +205,7 @@ private:
 
 	/** Given the duration of the call to Tick, perform the configured excess tick check. */
 	void EvaluatePlatformTickTime(double&& TickTime);
+	bool CheckStepTimedOut(FStepPtr& Step, double CurrentTime, IOnlineSubsystem* Subsystem);
 
 
 	FTestDriver& Driver;
