- Fix Heirarchy typo

#rnx

diff --git a/Engine/Plugins/Online/OnlineServices/Source/OnlineServicesCommon/Public/Online/OnlineAsyncOpCache.h b/Engine/Plugins/Online/OnlineServices/Source/OnlineServicesCommon/Public/Online/OnlineAsyncOpCache.h
--- a/Engine/Plugins/Online/OnlineServices/Source/OnlineServicesCommon/Public/Online/OnlineAsyncOpCache.h
+++ b/Engine/Plugins/Online/OnlineServices/Source/OnlineServicesCommon/Public/Online/OnlineAsyncOpCache.h
@@ -246,7 +246,7 @@ public:
 
 	// Create an operation
 	template <typename OpType>
-	TOnlineAsyncOpRef<OpType> GetOp(typename OpType::Params&& Params, const TArray<FString>& ConfigSectionHeiarchy)
+	TOnlineAsyncOpRef<OpType> GetOp(typename OpType::Params&& Params, const TArray<FString>& ConfigSectionHeirarchy)
 	{
 		TUniquePtr<TWrappedOperation<OpType>> WrappedOp = MakeUnique<TWrappedOperation<OpType>>(Services, MoveTemp(Params));
 		TSharedRef<TOnlineAsyncOp<OpType>> Op = WrappedOp->GetDataRef();
@@ -276,7 +276,7 @@ public:
 
 	// Join an existing operation or use a non-expired cached result, or create an operation that can later be joined
 	template <typename OpType, typename ParamsFuncsType /*= TJoinableOpParamsFuncs<OpType>*/>
-	TOnlineAsyncOpRef<OpType> GetJoinableOp(typename OpType::Params&& Params, const TArray<FString>& ConfigSectionHeiarchy)
+	TOnlineAsyncOpRef<OpType> GetJoinableOp(typename OpType::Params&& Params, const TArray<FString>& ConfigSectionHeirarchy)
 	{
 		TSharedPtr<TOnlineAsyncOp<OpType>> Op;
 
@@ -306,7 +306,7 @@ public:
 		{
 			FOperationConfig Config;
 			check(LoadConfigFn);
-			LoadConfigFn(Config, ConfigSectionHeiarchy);
+			LoadConfigFn(Config, ConfigSectionHeirarchy);
 			
 			Op = CreateOp<OpType, ParamsFuncsType>(MoveTemp(Params));
 			Op->OpCacheHandle = Op->OnComplete().Add(GetSharedThis(), [this, Config](const TOnlineAsyncOp<OpType>& ThisOp, const TOnlineResult<OpType>& Result)
@@ -352,7 +352,7 @@ public:
 
 	// Merge with a pending operation, or create an operation
 	template <typename OpType, typename ParamsFuncsType /*= TMergeableOpParamsFuncs<OpType>*/>
-	TOnlineAsyncOpRef<OpType> GetMergeableOp(typename OpType::Params&& Params, const TArray<FString>& ConfigSectionHeiarchy)
+	TOnlineAsyncOpRef<OpType> GetMergeableOp(typename OpType::Params&& Params, const TArray<FString>& ConfigSectionHeirarchy)
 	{
 		TSharedPtr<TOnlineAsyncOp<OpType>> Op;
 
@@ -380,7 +380,7 @@ public:
 		{
 			FOperationConfig Config;
 			check(LoadConfigFn);
-			LoadConfigFn(Config, ConfigSectionHeiarchy);
+			LoadConfigFn(Config, ConfigSectionHeirarchy);
 
 			Op = CreateOp<OpType, ParamsFuncsType>(MoveTemp(Params));
 			// remove from cache once operation has started. It is no longer mergeable at that point
diff --git a/Engine/Plugins/Online/OnlineServices/Source/OnlineServicesCommon/Public/Online/OnlineComponent.h b/Engine/Plugins/Online/OnlineServices/Source/OnlineServicesCommon/Public/Online/OnlineComponent.h
--- a/Engine/Plugins/Online/OnlineServices/Source/OnlineServicesCommon/Public/Online/OnlineComponent.h
+++ b/Engine/Plugins/Online/OnlineServices/Source/OnlineServicesCommon/Public/Online/OnlineComponent.h
@@ -52,19 +52,19 @@ public:
 	template <typename OpType>
 	TOnlineAsyncOpRef<OpType> GetOp(typename OpType::Params&& Params)
 	{
-		return Services.template GetOp<OpType>(MoveTemp(Params), GetConfigSectionHeiarchy(OpType::Name));
+		return Services.template GetOp<OpType>(MoveTemp(Params), GetConfigSectionHeirarchy(OpType::Name));
 	}
 
 	template <typename OpType, typename ParamsFuncsType = TJoinableOpParamsFuncs<OpType>>
 	TOnlineAsyncOpRef<OpType> GetJoinableOp(typename OpType::Params&& Params)
 	{
-		return Services.template GetJoinableOp<OpType, ParamsFuncsType>(MoveTemp(Params), GetConfigSectionHeiarchy(OpType::Name));
+		return Services.template GetJoinableOp<OpType, ParamsFuncsType>(MoveTemp(Params), GetConfigSectionHeirarchy(OpType::Name));
 	}
 
 	template <typename OpType, typename ParamsFuncsType = TMergeableOpParamsFuncs<OpType>>
 	TOnlineAsyncOpRef<OpType> GetMergeableOp(typename OpType::Params&& Params)
 	{
-		return Services.template GetMergeableOp<OpType, ParamsFuncsType>(MoveTemp(Params), GetConfigSectionHeiarchy(OpType::Name));
+		return Services.template GetMergeableOp<OpType, ParamsFuncsType>(MoveTemp(Params), GetConfigSectionHeirarchy(OpType::Name));
 	}
 
 	template <typename ServicesType = FOnlineServicesCommon>
@@ -97,21 +97,27 @@ public:
 		return **Queue;
 	}
 
+	UE_DEPRECATED(5.6, "GetConfigSectionHeiarchy has been renamed GetConfigSectionHeirarchy")
 	TArray<FString> GetConfigSectionHeiarchy(const FString& OperationName = FString())
 	{
-		TArray<FString> SectionHeiarchy;
+		return GetConfigSectionHeirarchy(OperationName);
+	}
+
+	TArray<FString> GetConfigSectionHeirarchy(const FString& OperationName = FString())
+	{
+		TArray<FString> SectionHeirarchy;
 		FString SectionName = TEXT("OnlineServices");
-		SectionHeiarchy.Add(SectionName);
+		SectionHeirarchy.Add(SectionName);
 		SectionName += TEXT(".") + GetServices().GetServiceConfigName();
-		SectionHeiarchy.Add(SectionName);
+		SectionHeirarchy.Add(SectionName);
 		SectionName += TEXT(".") + GetConfigName();
-		SectionHeiarchy.Add(SectionName);
+		SectionHeirarchy.Add(SectionName);
 		if (!OperationName.IsEmpty())
 		{
 			SectionName += TEXT(".") + OperationName;
-			SectionHeiarchy.Add(SectionName);
+			SectionHeirarchy.Add(SectionName);
 		}
-		return SectionHeiarchy;
+		return SectionHeirarchy;
 	}
 
 	virtual void RegisterCommands() {}
diff --git a/Engine/Plugins/Online/OnlineServices/Source/OnlineServicesCommon/Public/Online/OnlineServicesCommon.h b/Engine/Plugins/Online/OnlineServices/Source/OnlineServicesCommon/Public/Online/OnlineServicesCommon.h
--- a/Engine/Plugins/Online/OnlineServices/Source/OnlineServicesCommon/Public/Online/OnlineServicesCommon.h
+++ b/Engine/Plugins/Online/OnlineServices/Source/OnlineServicesCommon/Public/Online/OnlineServicesCommon.h
@@ -183,19 +183,25 @@ public:
 	const FString& GetConfigName() const { return ServiceConfigName; }
 	const FString& GetServiceConfigName() const { return ServiceConfigName; }
 
+	UE_DEPRECATED(5.6, "GetConfigSectionHeiarchy has been renamed GetConfigSectionHeirarchy")
 	TArray<FString> GetConfigSectionHeiarchy(const FString& OperationName = FString()) const
 	{
-		TArray<FString> SectionHeiarchy;
+		return GetConfigSectionHeirarchy(OperationName);
+	}
+	
+	TArray<FString> GetConfigSectionHeirarchy(const FString& OperationName = FString()) const
+	{
+		TArray<FString> SectionHeirarchy;
 		FString SectionName = TEXT("OnlineServices");
-		SectionHeiarchy.Add(SectionName);
+		SectionHeirarchy.Add(SectionName);
 		SectionName += TEXT(".") + GetServiceConfigName();
-		SectionHeiarchy.Add(SectionName);
+		SectionHeirarchy.Add(SectionName);
 		if (!OperationName.IsEmpty())
 		{
 			SectionName += TEXT(".") + OperationName;
-			SectionHeiarchy.Add(SectionName);
+			SectionHeirarchy.Add(SectionName);
 		}
-		return SectionHeiarchy;
+		return SectionHeirarchy;
 	}
 
 	/**
@@ -216,36 +222,42 @@ public:
 	template <typename StructType>
 	bool LoadConfig(StructType& Struct, const FString& InterfaceName = FString(), const FString& OperationName = FString()) const
 	{
-		TArray<FString> SectionHeiarchy;
+		TArray<FString> SectionHeirarchy;
 		FString SectionName = TEXT("OnlineServices");
-		SectionHeiarchy.Add(SectionName);
+		SectionHeirarchy.Add(SectionName);
 		if (!InterfaceName.IsEmpty())
 		{
-			SectionHeiarchy.Add(SectionName + TEXT(".") + InterfaceName);
+			SectionHeirarchy.Add(SectionName + TEXT(".") + InterfaceName);
 		}
 		SectionName += TEXT(".") + GetServiceConfigName();
-		SectionHeiarchy.Add(SectionName);
+		SectionHeirarchy.Add(SectionName);
 		if (!InterfaceName.IsEmpty())
 		{
 			SectionName += TEXT(".") + InterfaceName;
-			SectionHeiarchy.Add(SectionName);
+			SectionHeirarchy.Add(SectionName);
 			if (!OperationName.IsEmpty())
 			{
 				SectionName += TEXT(".") + OperationName;
-				SectionHeiarchy.Add(SectionName);
+				SectionHeirarchy.Add(SectionName);
 			}
 		}
-		return LoadConfig(Struct, SectionHeiarchy);
+		return LoadConfig(Struct, SectionHeirarchy);
+	}
+	
+	UE_DEPRECATED(5.6, "GetConfigSectionHeirachWithOverrides has been renamed GetConfigSectionHeirarchyWithOverrides")
+	TArray<FString> GetConfigSectionHeirachWithOverrides(const TArray<FString>& SectionHeirarchy) const
+	{
+		return GetConfigSectionHeirarchyWithOverrides(SectionHeirarchy);
 	}
 
 	/**
 	 * Get an array of a config section with the overrides added in
 	 * 
-	 * @param SectionHeiarchy Array of config sections to load values from
+	 * @param SectionHeirarchy Array of config sections to load values from
 	 * 
 	 * @return Array of the sections with overrides for values to be loaded from
 	 */
-	TArray<FString> GetConfigSectionHeirachWithOverrides(const TArray<FString>& SectionHeiarchy) const
+	TArray<FString> GetConfigSectionHeirarchyWithOverrides(const TArray<FString>& SectionHeirarchy) const
 	{
 		TArray<FString> AllConfigSectionOverrides = ConfigSectionOverrides;
 
@@ -254,37 +266,37 @@ public:
 			AllConfigSectionOverrides.Add("Editor");
 		}
 
-		TArray<FString> SectionHeiarchyWithOverrides;
-		for (const FString& Section : SectionHeiarchy)
+		TArray<FString> SectionHeirarchyWithOverrides;
+		for (const FString& Section : SectionHeirarchy)
 		{
-			SectionHeiarchyWithOverrides.Add(Section);
+			SectionHeirarchyWithOverrides.Add(Section);
 
 			for (const FString& Override : AllConfigSectionOverrides)
 			{
 				// Using of space in online config section is deprecated
 				FString OverrideSection = Section + TEXT(" ") + Override;
-				SectionHeiarchyWithOverrides.Add(OverrideSection);
+				SectionHeirarchyWithOverrides.Add(OverrideSection);
 
 				OverrideSection = Section + TEXT(":") + Override;
-				SectionHeiarchyWithOverrides.Add(OverrideSection);
+				SectionHeirarchyWithOverrides.Add(OverrideSection);
 			}
 		}
 
-		return SectionHeiarchyWithOverrides;
+		return SectionHeirarchyWithOverrides;
 	}
 
 	/**
-	 * Load a config struct for a section heiarchy, also using the ConfigSectionOverrides
+	 * Load a config struct for a section heirarchy, also using the ConfigSectionOverrides
 	 *
 	 * @param Struct Struct to populate with values from config
-	 * @param SectionHeiarchy Array of config sections to load values from
+	 * @param SectionHeirarchy Array of config sections to load values from
 	 *
 	 * @return true if a value was loaded
 	 */
 	template <typename StructType>
-	bool LoadConfig(StructType& Struct, const TArray<FString>& SectionHeiarchy) const
+	bool LoadConfig(StructType& Struct, const TArray<FString>& SectionHeirarchy) const
 	{
-		return ::UE::Online::LoadConfig(*ConfigProvider, GetConfigSectionHeirachWithOverrides(SectionHeiarchy), Struct);
+		return ::UE::Online::LoadConfig(*ConfigProvider, GetConfigSectionHeirarchyWithOverrides(SectionHeirarchy), Struct);
 	}
 
 	/* Get op (OnlineServices) */
@@ -292,43 +304,43 @@ public:
 	TOnlineAsyncOpRef<OpType> GetOp(typename OpType::Params&& Params)
 	{
 		LogErrorIfOnlineServicesHasShutdown<OpType>();
-		return OpCache.GetOp<OpType>(MoveTemp(Params), GetConfigSectionHeiarchy());
+		return OpCache.GetOp<OpType>(MoveTemp(Params), GetConfigSectionHeirarchy());
 	}
 
 	template <typename OpType, typename ParamsFuncsType = TJoinableOpParamsFuncs<OpType>>
 	TOnlineAsyncOpRef<OpType> GetJoinableOp(typename OpType::Params&& Params)
 	{
 		LogErrorIfOnlineServicesHasShutdown<OpType>();
-		return OpCache.GetJoinableOp<OpType, ParamsFuncsType>(MoveTemp(Params), GetConfigSectionHeiarchy());
+		return OpCache.GetJoinableOp<OpType, ParamsFuncsType>(MoveTemp(Params), GetConfigSectionHeirarchy());
 	}
 
 	template <typename OpType, typename ParamsFuncsType = TMergeableOpParamsFuncs<OpType>>
 	TOnlineAsyncOpRef<OpType> GetMergeableOp(typename OpType::Params&& Params)
 	{
 		LogErrorIfOnlineServicesHasShutdown<OpType>();
-		return OpCache.GetMergeableOp<OpType, ParamsFuncsType>(MoveTemp(Params), GetConfigSectionHeiarchy());
+		return OpCache.GetMergeableOp<OpType, ParamsFuncsType>(MoveTemp(Params), GetConfigSectionHeirarchy());
 	}
 
 	/* Get op (Interface) */
 	template <typename OpType>
-	TOnlineAsyncOpRef<OpType> GetOp(typename OpType::Params&& Params, const TArray<FString>& ConfigSectionHeiarchy)
+	TOnlineAsyncOpRef<OpType> GetOp(typename OpType::Params&& Params, const TArray<FString>& ConfigSectionHeirarchy)
 	{
 		LogErrorIfOnlineServicesHasShutdown<OpType>();
-		return OpCache.GetOp<OpType>(MoveTemp(Params), ConfigSectionHeiarchy);
+		return OpCache.GetOp<OpType>(MoveTemp(Params), ConfigSectionHeirarchy);
 	}
 
 	template <typename OpType, typename ParamsFuncsType = TJoinableOpParamsFuncs<OpType>>
-	TOnlineAsyncOpRef<OpType> GetJoinableOp(typename OpType::Params&& Params, const TArray<FString>& ConfigSectionHeiarchy)
+	TOnlineAsyncOpRef<OpType> GetJoinableOp(typename OpType::Params&& Params, const TArray<FString>& ConfigSectionHeirarchy)
 	{
 		LogErrorIfOnlineServicesHasShutdown<OpType>();
-		return OpCache.GetJoinableOp<OpType, ParamsFuncsType>(MoveTemp(Params), ConfigSectionHeiarchy);
+		return OpCache.GetJoinableOp<OpType, ParamsFuncsType>(MoveTemp(Params), ConfigSectionHeirarchy);
 	}
 
 	template <typename OpType, typename ParamsFuncsType = TMergeableOpParamsFuncs<OpType>>
-	TOnlineAsyncOpRef<OpType> GetMergeableOp(typename OpType::Params&& Params, const TArray<FString>& ConfigSectionHeiarchy)
+	TOnlineAsyncOpRef<OpType> GetMergeableOp(typename OpType::Params&& Params, const TArray<FString>& ConfigSectionHeirarchy)
 	{
 		LogErrorIfOnlineServicesHasShutdown<OpType>();
-		return OpCache.GetMergeableOp<OpType, ParamsFuncsType>(MoveTemp(Params), ConfigSectionHeiarchy);
+		return OpCache.GetMergeableOp<OpType, ParamsFuncsType>(MoveTemp(Params), ConfigSectionHeirarchy);
 	}
 
 	/* Queue for executing tasks in parallel. Serial queues feed into this */
diff --git a/Engine/Source/Programs/Online/OnlineTestsCore/Source/Private/Tests/TestAsyncOp.cpp b/Engine/Source/Programs/Online/OnlineTestsCore/Source/Private/Tests/TestAsyncOp.cpp
--- a/Engine/Source/Programs/Online/OnlineTestsCore/Source/Private/Tests/TestAsyncOp.cpp
+++ b/Engine/Source/Programs/Online/OnlineTestsCore/Source/Private/Tests/TestAsyncOp.cpp
@@ -253,24 +253,24 @@ ASYNCOP_TEST_CASE("Verify that the various OnlineAsyncOp execution policies work
 ASYNCOP_TEST_CASE("Verify that the duration cache policy works)")
 {
 	FTestPipeline& TestPipeline = GetPipeline();
-	TArray<FString> CacheDurationTestHeiarchy = { TEXT("OnlineTests.TestCacheDuration") };
+	TArray<FString> CacheDurationTestHeirarchy = { TEXT("OnlineTests.TestCacheDuration") };
 	// Uses the following operation configuration in OnlineTests' DefaultEngine.ini
 	// [OnlineTests.TestCacheDuration]
 	// CacheExpiration=Duration
 	// CacheExpirySeconds=2
 	// bCacheError=false
-	TestPipeline.EmplaceLambda([&CacheDurationTestHeiarchy](const IOnlineServicesPtr& Services)
+	TestPipeline.EmplaceLambda([&CacheDurationTestHeirarchy](const IOnlineServicesPtr& Services)
 		{
 			UE::Online::FOnlineServicesCommon& ServicesCommon = static_cast<UE::Online::FOnlineServicesCommon&>(*Services);
-			TOnlineAsyncOpRef<FTestOp> Op = ServicesCommon.GetJoinableOp<FTestOp>(FTestOp::Params(), CacheDurationTestHeiarchy);
+			TOnlineAsyncOpRef<FTestOp> Op = ServicesCommon.GetJoinableOp<FTestOp>(FTestOp::Params(), CacheDurationTestHeirarchy);
 			Op->Then([](FOnlineAsyncOp&) {})
 				.Then([](TOnlineAsyncOp<FTestOp>& Op) { Op.SetResult(FTestOp::Result()); })
 				.Enqueue(ServicesCommon.GetParallelQueue());
 		})
-		.EmplaceLambda([&CacheDurationTestHeiarchy](const IOnlineServicesPtr& Services)
+		.EmplaceLambda([&CacheDurationTestHeirarchy](const IOnlineServicesPtr& Services)
 		{
 			UE::Online::FOnlineServicesCommon& ServicesCommon = static_cast<UE::Online::FOnlineServicesCommon&>(*Services);
-			TOnlineAsyncOpRef<FTestOp> Op = ServicesCommon.GetJoinableOp<FTestOp>(FTestOp::Params(), CacheDurationTestHeiarchy);
+			TOnlineAsyncOpRef<FTestOp> Op = ServicesCommon.GetJoinableOp<FTestOp>(FTestOp::Params(), CacheDurationTestHeirarchy);
 			CHECK(Op->IsReady()); // Cache should return the already in flight/completed op
 		})
 		.EmplaceAsyncLambda([](FAsyncLambdaResult Promise, const IOnlineServicesPtr& Services)
@@ -278,10 +278,10 @@ ASYNCOP_TEST_CASE("Verify that the duration cache policy works)")
 			// Wait a second into the cache duration
 			FTSTicker::GetCoreTicker().AddTicker(TEXT("OnlineTestTimeout"), 1.0f, [Promise](float) { Promise->SetValue(true); return false; });
 		})
-		.EmplaceLambda([&CacheDurationTestHeiarchy](const IOnlineServicesPtr& Services)
+		.EmplaceLambda([&CacheDurationTestHeirarchy](const IOnlineServicesPtr& Services)
 		{
 			UE::Online::FOnlineServicesCommon& ServicesCommon = static_cast<UE::Online::FOnlineServicesCommon&>(*Services);
-			TOnlineAsyncOpRef<FTestOp> Op = ServicesCommon.GetJoinableOp<FTestOp>(FTestOp::Params(), CacheDurationTestHeiarchy);
+			TOnlineAsyncOpRef<FTestOp> Op = ServicesCommon.GetJoinableOp<FTestOp>(FTestOp::Params(), CacheDurationTestHeirarchy);
 			CHECK(Op->IsReady()); // Cache should return the already completed op
 		})
 		.EmplaceAsyncLambda([](FAsyncLambdaResult Promise, const IOnlineServicesPtr& Services)
@@ -289,10 +289,10 @@ ASYNCOP_TEST_CASE("Verify that the duration cache policy works)")
 			// Wait until shortly after the cache should expire
 			FTSTicker::GetCoreTicker().AddTicker(TEXT("OnlineTestTimeout"), 2.0f, [Promise](float) { Promise->SetValue(true); return false; });
 		})
-		.EmplaceAsyncLambda([&CacheDurationTestHeiarchy](FAsyncLambdaResult Promise, const IOnlineServicesPtr& Services)
+		.EmplaceAsyncLambda([&CacheDurationTestHeirarchy](FAsyncLambdaResult Promise, const IOnlineServicesPtr& Services)
 		{
 			UE::Online::FOnlineServicesCommon& ServicesCommon = static_cast<UE::Online::FOnlineServicesCommon&>(*Services);
-			TOnlineAsyncOpRef<FTestOp> Op = ServicesCommon.GetJoinableOp<FTestOp>(FTestOp::Params(), CacheDurationTestHeiarchy);
+			TOnlineAsyncOpRef<FTestOp> Op = ServicesCommon.GetJoinableOp<FTestOp>(FTestOp::Params(), CacheDurationTestHeirarchy);
 			CHECK(!Op->IsReady()); // Cache should return a new op
 			Op->Then([](FOnlineAsyncOp&) {})
 				.Then([Promise](TOnlineAsyncOp<FTestOp>& Op) { Op.SetResult(FTestOp::Result()); Promise->SetValue(true); })
@@ -311,39 +311,39 @@ ASYNCOP_TEST_CASE("Verify that the cache error setting works)", ASYNCOP_SKIP_EOS
 {
 	DestroyCurrentServiceModule();
 	FTestPipeline& TestPipeline = GetPipeline();
-	TArray<FString> DontCacheErrorHeiarchy = { TEXT("OnlineTests.TestDontCacheError") };
+	TArray<FString> DontCacheErrorHeirarchy = { TEXT("OnlineTests.TestDontCacheError") };
 	// Uses the following operation configuration in OnlineTests' DefaultEngine.ini
 	// [OnlineTests.TestDontCacheError]
 	// CacheExpiration=Duration
 	// CacheExpirySeconds=1
 	// bCacheError=false
-	TArray<FString> CacheErrorHeiarchy = { TEXT("OnlineTests.TestCacheError") };
+	TArray<FString> CacheErrorHeirarchy = { TEXT("OnlineTests.TestCacheError") };
 	// Uses the following operation configuration in OnlineTests' DefaultEngine.ini
 	// [OnlineTests.TestCacheError]
 	// CacheExpiration=Duration
 	// CacheExpirySeconds=1
 	// bCacheError=true
-	TestPipeline.EmplaceAsyncLambda([&DontCacheErrorHeiarchy](FAsyncLambdaResult Promise, const IOnlineServicesPtr& Services)
+	TestPipeline.EmplaceAsyncLambda([&DontCacheErrorHeirarchy](FAsyncLambdaResult Promise, const IOnlineServicesPtr& Services)
 		{
 			UE::Online::FOnlineServicesCommon& ServicesCommon = static_cast<UE::Online::FOnlineServicesCommon&>(*Services);
-			TOnlineAsyncOpRef<FTestOp> Op = ServicesCommon.GetJoinableOp<FTestOp>(FTestOp::Params(), DontCacheErrorHeiarchy);
+			TOnlineAsyncOpRef<FTestOp> Op = ServicesCommon.GetJoinableOp<FTestOp>(FTestOp::Params(), DontCacheErrorHeirarchy);
 			Op->Then([](FOnlineAsyncOp&) {})
 				.Then([Promise](TOnlineAsyncOp<FTestOp>& Op) { Op.SetError(UE::Online::Errors::Unknown()); Promise->SetValue(true); })
 				.Enqueue(ServicesCommon.GetParallelQueue());
 		})
-		.EmplaceAsyncLambda([&CacheErrorHeiarchy](FAsyncLambdaResult Promise, const IOnlineServicesPtr& Services)
+		.EmplaceAsyncLambda([&CacheErrorHeirarchy](FAsyncLambdaResult Promise, const IOnlineServicesPtr& Services)
 		{
 			UE::Online::FOnlineServicesCommon& ServicesCommon = static_cast<UE::Online::FOnlineServicesCommon&>(*Services);
-			TOnlineAsyncOpRef<FTestOp> Op = ServicesCommon.GetJoinableOp<FTestOp>(FTestOp::Params(), CacheErrorHeiarchy);
+			TOnlineAsyncOpRef<FTestOp> Op = ServicesCommon.GetJoinableOp<FTestOp>(FTestOp::Params(), CacheErrorHeirarchy);
 			CHECK(!Op->IsReady()); // Should return a new op
 			Op->Then([](FOnlineAsyncOp&) {})
 				.Then([Promise](TOnlineAsyncOp<FTestOp>& Op) { Op.SetError(UE::Online::Errors::Unknown()); Promise->SetValue(true); })
 				.Enqueue(ServicesCommon.GetParallelQueue());
 		})
-		.EmplaceAsyncLambda([&CacheErrorHeiarchy](FAsyncLambdaResult Promise, const IOnlineServicesPtr& Services)
+		.EmplaceAsyncLambda([&CacheErrorHeirarchy](FAsyncLambdaResult Promise, const IOnlineServicesPtr& Services)
 		{
 			UE::Online::FOnlineServicesCommon& ServicesCommon = static_cast<UE::Online::FOnlineServicesCommon&>(*Services);
-			TOnlineAsyncOpRef<FTestOp> Op = ServicesCommon.GetJoinableOp<FTestOp>(FTestOp::Params(), CacheErrorHeiarchy);
+			TOnlineAsyncOpRef<FTestOp> Op = ServicesCommon.GetJoinableOp<FTestOp>(FTestOp::Params(), CacheErrorHeirarchy);
 			CHECK(Op->IsReady()); // Should return the cached, failed op
 			Op->GetHandle().OnComplete([Promise](const TOnlineResult<FTestOp>& Result)
 				{
@@ -363,26 +363,26 @@ ASYNCOP_TEST_CASE("Verify that the cache error setting works)", ASYNCOP_SKIP_EOS
 ASYNCOP_TEST_CASE("Verify that the cache expires once complete)")
 {
 	FTestPipeline& TestPipeline = GetPipeline();
-	TArray<FString> CacheExpireUponCompletionHeiarchy = { TEXT("OnlineTests.TestCacheExpireUponCompletion") };
+	TArray<FString> CacheExpireUponCompletionHeirarchy = { TEXT("OnlineTests.TestCacheExpireUponCompletion") };
 	// Uses the following operation configuration in OnlineTests' DefaultEngine.ini
 	// [OnlineTests.TestCacheExpireUponCompletion]
 	// CacheExpiration=UponCompletion
 	// bCacheError=false
-	TestPipeline.EmplaceAsyncLambda([&CacheExpireUponCompletionHeiarchy](FAsyncLambdaResult Promise, const IOnlineServicesPtr& Services)
+	TestPipeline.EmplaceAsyncLambda([&CacheExpireUponCompletionHeirarchy](FAsyncLambdaResult Promise, const IOnlineServicesPtr& Services)
 		{
 			UE::Online::FOnlineServicesCommon& ServicesCommon = static_cast<UE::Online::FOnlineServicesCommon&>(*Services);
-			TOnlineAsyncOpRef<FTestOp> Op = ServicesCommon.GetJoinableOp<FTestOp>(FTestOp::Params(), CacheExpireUponCompletionHeiarchy);
+			TOnlineAsyncOpRef<FTestOp> Op = ServicesCommon.GetJoinableOp<FTestOp>(FTestOp::Params(), CacheExpireUponCompletionHeirarchy);
 			Op->Then([](FOnlineAsyncOp&) {})
 				.Then([Promise](TOnlineAsyncOp<FTestOp>& Op) { Op.SetResult(FTestOp::Result()); Promise->SetValue(true); })
 				.Enqueue(ServicesCommon.GetParallelQueue());
-			TOnlineAsyncOpRef<FTestOp> Op2 = ServicesCommon.GetJoinableOp<FTestOp>(FTestOp::Params(), CacheExpireUponCompletionHeiarchy);
+			TOnlineAsyncOpRef<FTestOp> Op2 = ServicesCommon.GetJoinableOp<FTestOp>(FTestOp::Params(), CacheExpireUponCompletionHeirarchy);
 			CHECK(Op2->IsReady());
 			CHECK(!Op2->IsComplete());
 		})
-		.EmplaceAsyncLambda([&CacheExpireUponCompletionHeiarchy](FAsyncLambdaResult Promise, const IOnlineServicesPtr& Services)
+		.EmplaceAsyncLambda([&CacheExpireUponCompletionHeirarchy](FAsyncLambdaResult Promise, const IOnlineServicesPtr& Services)
 		{
 			UE::Online::FOnlineServicesCommon& ServicesCommon = static_cast<UE::Online::FOnlineServicesCommon&>(*Services);
-			TOnlineAsyncOpRef<FTestOp> Op = ServicesCommon.GetJoinableOp<FTestOp>(FTestOp::Params(), CacheExpireUponCompletionHeiarchy);
+			TOnlineAsyncOpRef<FTestOp> Op = ServicesCommon.GetJoinableOp<FTestOp>(FTestOp::Params(), CacheExpireUponCompletionHeirarchy);
 			CHECK(!Op->IsReady());
 			Op->Then([](FOnlineAsyncOp&) {})
 				.Then([Promise](TOnlineAsyncOp<FTestOp>& Op) { Op.SetResult(FTestOp::Result()); Promise->SetValue(true); })
