Fix the warnings in OnlineTests related to ossv2 shutdown, by calling DestroyCurrentServiceModule in RunToCompletion make sure online services is destroyed properly in each test case.

Reason: Online services instances got created when run a test case. Each test case can run multiple times because of different config, for example, one for null, one for eos, one for platform. But online services instance wasn't destroyed each time. So it's still there for next test case, each test case doesn't get a clean environment by proper setup and shutdown.

diff --git a/Engine/Plugins/Online/OnlineServicesNull/Source/Private/Online/OnlineServicesNull.cpp b/Engine/Plugins/Online/OnlineServicesNull/Source/Private/Online/OnlineServicesNull.cpp
--- a/Engine/Plugins/Online/OnlineServicesNull/Source/Private/Online/OnlineServicesNull.cpp
+++ b/Engine/Plugins/Online/OnlineServicesNull/Source/Private/Online/OnlineServicesNull.cpp
@@ -48,6 +48,21 @@ void FOnlineServicesNull::RegisterComponents()
 	FOnlineServicesCommon::RegisterComponents();
 }
 
+void FOnlineServicesNull::PreShutdown()
+{
+	FOnlineServicesCommon::PreShutdown();
+
+	Components.Unregister<FAchievementsNull>();
+	Components.Unregister<FAuthNull>();
+	Components.Unregister<FPresenceNull>();
+	Components.Unregister<FLeaderboardsNull>();
+	Components.Unregister<FLobbiesNull>();
+	Components.Unregister<FStatsNull>();
+	Components.Unregister<FTitleFileNull>();
+	Components.Unregister<FUserFileNull>();
+	Components.Unregister<FSessionsNull>();
+}
+
 TOnlineResult<FGetResolvedConnectString> FOnlineServicesNull::GetResolvedConnectString(FGetResolvedConnectString::Params&& Params)
 {
 	if(Params.LobbyId.IsValid())
diff --git a/Engine/Plugins/Online/OnlineServicesNull/Source/Public/Online/OnlineServicesNull.h b/Engine/Plugins/Online/OnlineServicesNull/Source/Public/Online/OnlineServicesNull.h
--- a/Engine/Plugins/Online/OnlineServicesNull/Source/Public/Online/OnlineServicesNull.h
+++ b/Engine/Plugins/Online/OnlineServicesNull/Source/Public/Online/OnlineServicesNull.h
@@ -16,6 +16,7 @@ public:
 	UE_API FOnlineServicesNull(FName InInstanceName, FName InInstanceConfigName);
 	UE_API virtual void RegisterComponents() override;
 	UE_API virtual void Initialize() override;
+	UE_API virtual void PreShutdown() override;
 	UE_API virtual TOnlineResult<FGetResolvedConnectString> GetResolvedConnectString(FGetResolvedConnectString::Params&& Params) override;
 	virtual EOnlineServices GetServicesProvider() const override { return EOnlineServices::Null; }
 
diff --git a/Engine/Source/Programs/Online/OnlineTestsCore/Source/Private/OnlineCatchHelper.cpp b/Engine/Source/Programs/Online/OnlineTestsCore/Source/Private/OnlineCatchHelper.cpp
--- a/Engine/Source/Programs/Online/OnlineTestsCore/Source/Private/OnlineCatchHelper.cpp
+++ b/Engine/Source/Programs/Online/OnlineTestsCore/Source/Private/OnlineCatchHelper.cpp
@@ -297,6 +297,13 @@ void OnlineTestBase::RunToCompletion(bool bLogout) const
 	REQUIRE(Driver.AddPipeline(MoveTemp(*Pipeline), TestContext)); // If this fails, we were unable to find the subsystem that is being passed by GetService
 	Pipeline = nullptr;
 	Driver.RunToCompletion();
+
+	if (bLogout)
+	{
+		// This is to fix "XXX online services is still been referenced after shutting down". Also need to 
+		// make sure to not cache online services shared ptr or interface ptr in your test case scope
+		DestroyCurrentServiceModule();
+	}
 }
 
 TArray<OnlineAutoReg::FApplicableServicesConfig> OnlineAutoReg::GetApplicableServices()
diff --git a/Engine/Source/Programs/Online/OnlineTestsCore/Source/Private/Tests/Presence/PresenceTests.cpp b/Engine/Source/Programs/Online/OnlineTestsCore/Source/Private/Tests/Presence/PresenceTests.cpp
--- a/Engine/Source/Programs/Online/OnlineTestsCore/Source/Private/Tests/Presence/PresenceTests.cpp
+++ b/Engine/Source/Programs/Online/OnlineTestsCore/Source/Private/Tests/Presence/PresenceTests.cpp
@@ -45,7 +45,7 @@ PRESENCE_TEST_CASE("Ensure UpdatePresence works on multiple users and can then b
 	TArray<TSharedRef<const FUserPresence>> OutPresences;
 
 	GetLoginPipeline({ AccountIdA, AccountIdB })
-		.EmplaceLambda([&](const IOnlineServicesPtr& Type)
+		.EmplaceLambda([&AccountIdA, &AccountIdB, &AccountIds](const IOnlineServicesPtr& Type)
 		{
 			AccountIds.Add(AccountIdA);
 			AccountIds.Add(AccountIdB);
@@ -53,14 +53,13 @@ PRESENCE_TEST_CASE("Ensure UpdatePresence works on multiple users and can then b
 		.EmplaceStep<FUpdatePresenceHelper>(AccountIdA, GetGenericPresenceA(AccountIdA))
 		.EmplaceStep<FUpdatePresenceHelper>(AccountIdB, GetGenericPresenceB(AccountIdB))
 		.EmplaceStep<FBatchQueryPresenceHelper>(AccountIdA, AccountIds, OutPresences)
-		.EmplaceLambda([&](const IOnlineServicesPtr& Type)
+		.EmplaceLambda([&AccountIdA, &AccountIdB, &OutPresences](const IOnlineServicesPtr& Type)
 		{
 			REQUIRE(OutPresences.Num() == 2);
 			CheckPresenceAreEqual(*OutPresences[0], *GetGenericPresenceA(AccountIdA));
 			CheckPresenceAreEqual(*OutPresences[1], *GetGenericPresenceB(AccountIdB));
 		});
 
-
 	RunToCompletion();
 }
 
@@ -102,8 +101,7 @@ PRESENCE_TEST_CASE("Ensure QueryPresence's bListenToChanges mod fires the event
 	TSharedPtr<const FUserPresence> OutPresence;
 	bool bReceivedPresence = false;
 
-	IPresencePtr Ptr = GetSubsystem()->GetPresenceInterface();
-	auto Handle = Ptr->OnPresenceUpdated().Add([&](const FPresenceUpdated& Event) mutable
+	auto Handle = GetSubsystem()->GetPresenceInterface()->OnPresenceUpdated().Add([&](const FPresenceUpdated& Event) mutable
 	{
 		if (Event.UpdatedPresence->AccountId == AccountIdB)
 		{
@@ -136,8 +134,7 @@ PRESENCE_TEST_CASE("Ensure BatchQueryPresence's bListenToChanges mod fires the e
 
 	bool bReceivedPresence = false;
 
-	IPresencePtr Ptr = GetSubsystem()->GetPresenceInterface();
-	auto Handle = Ptr->OnPresenceUpdated().Add([&](const FPresenceUpdated& Event)
+	auto Handle = GetSubsystem()->GetPresenceInterface()->OnPresenceUpdated().Add([&](const FPresenceUpdated& Event)
 	{
 		if (Event.UpdatedPresence->AccountId == AccountIdB)
 		{
