Remove redundant DestroyCurrentServiceModule from all tests cases in OnlineTests.

diff --git a/Engine/Source/Programs/Online/OnlineTestsCore/Source/Private/Tests/Achievements/DisplayAchievementUITests.cpp b/Engine/Source/Programs/Online/OnlineTestsCore/Source/Private/Tests/Achievements/DisplayAchievementUITests.cpp
--- a/Engine/Source/Programs/Online/OnlineTestsCore/Source/Private/Tests/Achievements/DisplayAchievementUITests.cpp
+++ b/Engine/Source/Programs/Online/OnlineTestsCore/Source/Private/Tests/Achievements/DisplayAchievementUITests.cpp
@@ -29,8 +29,6 @@ ACHIEVEMENTS_TEST_CASE("Display Achievement UI (Invalid User)", EG_ACHIEVEMENTS_
 
 ACHIEVEMENTS_TEST_CASE("Display Achievement UI (Invalid State)",  EG_ACHIEVEMENTS_DISPLAYACHIEVEMENTUI_TAG)
 {
-	DestroyCurrentServiceModule();
-
 	FAccountId AccountId;
 	FQueryAchievementDefinitionsHelper::FHelperParams QueryDefinitionsHelperParams;
 
diff --git a/Engine/Source/Programs/Online/OnlineTestsCore/Source/Private/Tests/Achievements/GetAchievementDefinitionTests.cpp b/Engine/Source/Programs/Online/OnlineTestsCore/Source/Private/Tests/Achievements/GetAchievementDefinitionTests.cpp
--- a/Engine/Source/Programs/Online/OnlineTestsCore/Source/Private/Tests/Achievements/GetAchievementDefinitionTests.cpp
+++ b/Engine/Source/Programs/Online/OnlineTestsCore/Source/Private/Tests/Achievements/GetAchievementDefinitionTests.cpp
@@ -28,8 +28,6 @@ ACHIEVEMENTS_TEST_CASE("Get Achievement Definition (Invalid User)", EG_ACHIEVEME
 
 ACHIEVEMENTS_TEST_CASE("Get Achievement Definition (Invalid State)", EG_ACHIEVEMENTS_GETACHIEVEMENTDEFINITION_TAG)
 {
-	DestroyCurrentServiceModule();
-
 	FAccountId AccountId;
 
 	FTestPipeline& LoginPipeline = GetLoginPipeline({ AccountId });
@@ -75,8 +73,6 @@ ACHIEVEMENTS_TEST_CASE("Get Achievement Definition (Not Found)", EG_ACHIEVEMENTS
 
 ACHIEVEMENTS_TEST_CASE("Get Achievement Definitions (Success)", EG_ACHIEVEMENTS_GETACHIEVEMENTDEFINITION_TAG)
 {
-	DestroyCurrentServiceModule();
-
 	FAccountId AccountId;
 	TArray<FString> AchievementIds;
 
diff --git a/Engine/Source/Programs/Online/OnlineTestsCore/Source/Private/Tests/Achievements/GetAchievementIdsTests.cpp b/Engine/Source/Programs/Online/OnlineTestsCore/Source/Private/Tests/Achievements/GetAchievementIdsTests.cpp
--- a/Engine/Source/Programs/Online/OnlineTestsCore/Source/Private/Tests/Achievements/GetAchievementIdsTests.cpp
+++ b/Engine/Source/Programs/Online/OnlineTestsCore/Source/Private/Tests/Achievements/GetAchievementIdsTests.cpp
@@ -28,8 +28,6 @@ ACHIEVEMENTS_TEST_CASE("Get Achievement Ids (Invalid User)", EG_ACHIEVEMENTS_GET
 
 ACHIEVEMENTS_TEST_CASE("Get Achievement Ids (Invalid State)", EG_ACHIEVEMENTS_GETACHIEVEMENTIDS_TAG)
 {
-	DestroyCurrentServiceModule();
-
 	FAccountId AccountId;
 
 	auto& LoginPipeline = GetLoginPipeline({ AccountId });
diff --git a/Engine/Source/Programs/Online/OnlineTestsCore/Source/Private/Tests/Achievements/GetAchievementStateTests.cpp b/Engine/Source/Programs/Online/OnlineTestsCore/Source/Private/Tests/Achievements/GetAchievementStateTests.cpp
--- a/Engine/Source/Programs/Online/OnlineTestsCore/Source/Private/Tests/Achievements/GetAchievementStateTests.cpp
+++ b/Engine/Source/Programs/Online/OnlineTestsCore/Source/Private/Tests/Achievements/GetAchievementStateTests.cpp
@@ -14,8 +14,6 @@
 
 ACHIEVEMENTS_TEST_CASE("Get Achievement State (Invalid State)", EG_ACHIEVEMENTS_GETACHIEVEMENTSTATE_TAG)
 {
-	DestroyCurrentServiceModule();
-
 	FAccountId AccountId;
 
 	FTestPipeline& LoginPipeline = GetLoginPipeline({ AccountId });
diff --git a/Engine/Source/Programs/Online/OnlineTestsCore/Source/Private/Tests/Achievements/QueryAchievementStatesTests.cpp b/Engine/Source/Programs/Online/OnlineTestsCore/Source/Private/Tests/Achievements/QueryAchievementStatesTests.cpp
--- a/Engine/Source/Programs/Online/OnlineTestsCore/Source/Private/Tests/Achievements/QueryAchievementStatesTests.cpp
+++ b/Engine/Source/Programs/Online/OnlineTestsCore/Source/Private/Tests/Achievements/QueryAchievementStatesTests.cpp
@@ -27,8 +27,6 @@ ACHIEVEMENTS_TEST_CASE("Query Achievement States (Invalid User)", EG_ACHIEVEMENT
 
 ACHIEVEMENTS_TEST_CASE("Query Achievement States (Invalid State)", EG_ACHIEVEMENTS_QUERYACHIEVEMENTSTATES_TAG)
 {
-	DestroyCurrentServiceModule();
-
 	FAccountId AccountId;
 
 	FTestPipeline& LoginPipeline = GetLoginPipeline({ AccountId });
diff --git a/Engine/Source/Programs/Online/OnlineTestsCore/Source/Private/Tests/Achievements/UnlockAchievementsTests.cpp b/Engine/Source/Programs/Online/OnlineTestsCore/Source/Private/Tests/Achievements/UnlockAchievementsTests.cpp
--- a/Engine/Source/Programs/Online/OnlineTestsCore/Source/Private/Tests/Achievements/UnlockAchievementsTests.cpp
+++ b/Engine/Source/Programs/Online/OnlineTestsCore/Source/Private/Tests/Achievements/UnlockAchievementsTests.cpp
@@ -34,8 +34,6 @@ ACHIEVEMENTS_TEST_CASE("Unlock Achievements (Invalid User)", EG_ACHIEVEMENTS_UNL
 
 ACHIEVEMENTS_TEST_CASE("Unlock Achievements (Invalid State)", EG_ACHIEVEMENTS_UNLOCKACHIEVEMENTS_TAG)
 {
-	DestroyCurrentServiceModule();
-
 	FAccountId AccountId;
 
 	FTestPipeline& LoginPipeline = GetLoginPipeline({ AccountId });
@@ -104,8 +102,6 @@ ACHIEVEMENTS_TEST_CASE("Unlock Achievements (Not Found)", EG_ACHIEVEMENTS_UNLOCK
 
 ACHIEVEMENTS_TEST_CASE("Unlock Achievements (Already Unlocked)", EG_ACHIEVEMENTS_UNLOCKACHIEVEMENTS_TAG)
 {
-	DestroyCurrentServiceModule();
-
 	FAccountId AccountId;
 
 	FTestPipeline& LoginPipeline = GetLoginPipeline({ AccountId });
@@ -135,8 +131,6 @@ ACHIEVEMENTS_TEST_CASE("Unlock Achievements (Already Unlocked)", EG_ACHIEVEMENTS
 
 ACHIEVEMENTS_TEST_CASE("Unlock Achievements (Success)", EG_ACHIEVEMENTS_UNLOCKACHIEVEMENTS_TAG)
 {
-	DestroyCurrentServiceModule();
-
 	FAccountId AccountId;
 
 	FTestPipeline& LoginPipeline = GetLoginPipeline({ AccountId });
@@ -160,8 +154,6 @@ ACHIEVEMENTS_TEST_CASE("Unlock Achievements (Success)", EG_ACHIEVEMENTS_UNLOCKAC
 
 ACHIEVEMENTS_TEST_CASE("Unlock Title-Managed Achievements By Stats", EG_ACHIEVEMENTS_UNLOCKACHIEVEMENTS_TITLEMANAGED_TAG)
 {
-	DestroyCurrentServiceModule();
-
 	FAccountId AccountId;
 
 	FTestPipeline& LoginPipeline = GetLoginPipeline({ AccountId });
@@ -201,8 +193,6 @@ ACHIEVEMENTS_TEST_CASE("Unlock Title-Managed Achievements By Stats", EG_ACHIEVEM
 
 ACHIEVEMENTS_TEST_CASE("Unlock Title-Managed Achievements By Stats With Multiple Conditions", EG_ACHIEVEMENTS_UNLOCKACHIEVEMENTS_TITLEMANAGED_TAG)
 {
-	DestroyCurrentServiceModule();
-
 	FAccountId AccountId;
 
 	FTestPipeline& LoginPipeline = GetLoginPipeline({ AccountId });
@@ -251,5 +241,4 @@ ACHIEVEMENTS_TEST_CASE("Unlock Title-Managed Achievements By Stats With Multiple
 			});
 
 	RunToCompletion();
-	DestroyCurrentServiceModule();
 }
\ No newline at end of file
diff --git a/Engine/Source/Programs/Online/OnlineTestsCore/Source/Private/Tests/Leaderboards/LeaderboardsReadEntriesAroundUserTests.cpp b/Engine/Source/Programs/Online/OnlineTestsCore/Source/Private/Tests/Leaderboards/LeaderboardsReadEntriesAroundUserTests.cpp
--- a/Engine/Source/Programs/Online/OnlineTestsCore/Source/Private/Tests/Leaderboards/LeaderboardsReadEntriesAroundUserTests.cpp
+++ b/Engine/Source/Programs/Online/OnlineTestsCore/Source/Private/Tests/Leaderboards/LeaderboardsReadEntriesAroundUserTests.cpp
@@ -11,8 +11,6 @@
 
 LEADERBOARDS_TEST_CASE("Verify that ReadEntriesAroundUser succeeds and returns entries for the given board", "[MultiAccount]")
 {
-	DestroyCurrentServiceModule();
-
 	FAccountId AccountIdA, AccountIdB;
 	TArray< FAccountId*> AccountIds = { &AccountIdA, &AccountIdB };
 	TArray<UE::Online::FLeaderboardEntry> LeaderboardEntries;
@@ -97,8 +95,6 @@ LEADERBOARDS_TEST_CASE("Verify ReadEntriesAroundUser returns a number of entries
 
 LEADERBOARDS_TEST_CASE("Verify ReadEntriesAroundUser returns expected valid entries if the given Limit exceeds the number of existing entries before the center entry", "[MultiAccount]")
 {
-	DestroyCurrentServiceModule();
-
 	FAccountId LocalUser, AccountId;
 	TArray< FAccountId*> AccountIds = { &AccountId, &LocalUser };
 	int32 offset = -1;
@@ -114,8 +110,6 @@ LEADERBOARDS_TEST_CASE("Verify ReadEntriesAroundUser returns expected valid entr
 
 LEADERBOARDS_TEST_CASE("Verify ReadEntriesAroundUser returns expected valid entries if the given Limit exceeds the number of existing entries after the center entry", "[MultiAccount]")
 {
-	DestroyCurrentServiceModule();
-
 	FAccountId LocalUser, AccountId;
 	TArray< FAccountId*> AccountIds = { &AccountId, &LocalUser };
 	int32 offset = 0;
diff --git a/Engine/Source/Programs/Online/OnlineTestsCore/Source/Private/Tests/Leaderboards/LeaderboardsReadEntriesForUsersTests.cpp b/Engine/Source/Programs/Online/OnlineTestsCore/Source/Private/Tests/Leaderboards/LeaderboardsReadEntriesForUsersTests.cpp
--- a/Engine/Source/Programs/Online/OnlineTestsCore/Source/Private/Tests/Leaderboards/LeaderboardsReadEntriesForUsersTests.cpp
+++ b/Engine/Source/Programs/Online/OnlineTestsCore/Source/Private/Tests/Leaderboards/LeaderboardsReadEntriesForUsersTests.cpp
@@ -11,8 +11,6 @@
 
 LEADERBOARDS_TEST_CASE("ReadEntriesForUsers succeed", "[MultiAccount]")
 {
-	DestroyCurrentServiceModule();
-
 	TArray<UE::Online::FLeaderboardEntry> LeaderboardEntries;
 	FAccountId AccountIdA, AccountIdB;
 	TArray< FAccountId*> AccountIds = { &AccountIdA, &AccountIdB };
diff --git a/Engine/Source/Programs/Online/OnlineTestsCore/Source/Private/Tests/Sessions/AddRemoveSessionMemberTests.cpp b/Engine/Source/Programs/Online/OnlineTestsCore/Source/Private/Tests/Sessions/AddRemoveSessionMemberTests.cpp
--- a/Engine/Source/Programs/Online/OnlineTestsCore/Source/Private/Tests/Sessions/AddRemoveSessionMemberTests.cpp
+++ b/Engine/Source/Programs/Online/OnlineTestsCore/Source/Private/Tests/Sessions/AddRemoveSessionMemberTests.cpp
@@ -64,8 +64,6 @@ SESSIONS_TEST_CASE("If I call AddSessionMember with an unregistered session name
 
 SESSIONS_TEST_CASE("If I call AddSessionMember with valid data, the operation completes successfully", EG_SESSIONS_ADDREMOVESESSIONMEMBER_TAG)
 {
-	DestroyCurrentServiceModule();
-	
 	FAccountId AccountId;
 
 	FCreateSession::Params OpCreateParams;
@@ -155,8 +153,6 @@ SESSIONS_TEST_CASE("If I call RemoveSessionMember with an unregistered session n
 
 SESSIONS_TEST_CASE("If I call RemoveSessionMember with valid data, the operation completes successfully", EG_SESSIONS_ADDREMOVESESSIONMEMBER_TAG)
 {
-	DestroyCurrentServiceModule();
-
 	int32 TestAccountIndex = 7;
 	FAccountId AccountId;
 
diff --git a/Engine/Source/Programs/Online/OnlineTestsCore/Source/Private/Tests/Sessions/CreateSessionTests.cpp b/Engine/Source/Programs/Online/OnlineTestsCore/Source/Private/Tests/Sessions/CreateSessionTests.cpp
--- a/Engine/Source/Programs/Online/OnlineTestsCore/Source/Private/Tests/Sessions/CreateSessionTests.cpp
+++ b/Engine/Source/Programs/Online/OnlineTestsCore/Source/Private/Tests/Sessions/CreateSessionTests.cpp
@@ -148,8 +148,6 @@ SESSIONS_TEST_CASE("If I call CreateSession with session override id less than 1
 
 SESSIONS_TEST_CASE("If I call CreateSession with an name already in use, I get an error", EG_SESSIONS_CREATESESSION_TAG)
 {
-	DestroyCurrentServiceModule();
-	
 	int32 TestAccountIndex = 7;
 	FAccountId AccountId;
 
@@ -191,8 +189,6 @@ SESSIONS_TEST_CASE("If I call CreateSession with the presence flag set after the
 {
 	LOG_SCOPE_VERBOSITY_OVERRIDE(LogOnlineServices, ELogVerbosity::NoLogging);
 
-	DestroyCurrentServiceModule();
-	
 	int32 TestAccountIndex = 7;
 	FAccountId AccountId;
 
@@ -235,8 +231,6 @@ SESSIONS_TEST_CASE("If I call CreateSession with the presence flag set after the
 
 SESSIONS_TEST_CASE("If I call CreateSession with valid data, the operation completes successfully", EG_SESSIONS_CREATESESSION_TAG)
 {
-	DestroyCurrentServiceModule();
-
 	FAccountId AccountId;
 
 	FCreateSession::Params OpCreateParams;
@@ -267,8 +261,6 @@ SESSIONS_TEST_CASE("If I call CreateSession with valid data, the operation compl
 
 SESSIONS_TEST_CASE("If I call CreateSession twice for NULL, I get an error", EG_SESSIONS_CREATESESSIONNULL_TAG)
 {
-	DestroyCurrentServiceModule();
-
 	FAccountId FirstAccountId, SecondAccountId;
 
 	FCreateSession::Params OpFirstCreateParams;
diff --git a/Engine/Source/Programs/Online/OnlineTestsCore/Source/Private/Tests/Sessions/FindSessionsTests.cpp b/Engine/Source/Programs/Online/OnlineTestsCore/Source/Private/Tests/Sessions/FindSessionsTests.cpp
--- a/Engine/Source/Programs/Online/OnlineTestsCore/Source/Private/Tests/Sessions/FindSessionsTests.cpp
+++ b/Engine/Source/Programs/Online/OnlineTestsCore/Source/Private/Tests/Sessions/FindSessionsTests.cpp
@@ -135,7 +135,6 @@ SESSIONS_TEST_CASE("If I call FindSessions with an empty search filter key, I ge
 
 SESSIONS_TEST_CASE("If I call FindSessions while a search is already running for that user, I get an error", EG_SESSIONS_FINDSESSIONS_TAG)
 {
-	DestroyCurrentServiceModule();
 	ReturnAccounts();
 
 	int32 TestAccountIndex = 7;
@@ -213,8 +212,6 @@ SESSIONS_TEST_CASE("If I call FindSessions while a search is already running for
 
 SESSIONS_TEST_CASE("If I call FindSessions with valid information, I can get an empty list if no sessions are found", EG_SESSIONS_FINDSESSIONS_TAG)
 {
-	DestroyCurrentServiceModule();
-
 	FAccountId AccountId;
 
 	FFindSessions::Params OpFirstParams;
@@ -235,7 +232,6 @@ SESSIONS_TEST_CASE("If I call FindSessions with valid information, I can get an
 
 SESSIONS_TEST_CASE("If I call FindSessions with valid information, I will still not find sessions which JoinPolicy is not public", EG_SESSIONS_FINDSESSIONS_TAG)
 {
-	DestroyCurrentServiceModule();
 	ReturnAccounts();
 	
 	FAccountId FirstAccountId, SecondAccountId;
@@ -289,8 +285,6 @@ SESSIONS_TEST_CASE("If I call FindSessions with valid information, I will get a
 {
 	LOG_SCOPE_VERBOSITY_OVERRIDE(LogEOSSDK, ELogVerbosity::NoLogging);
 
-	DestroyCurrentServiceModule();
-
 	FAccountId FirstAccountId, SecondAccountId;
 
 	FCreateSession::Params OpFirstCreateParams;
@@ -369,8 +363,6 @@ SESSIONS_TEST_CASE("Call FindSessions using custom setting as a filter Bool", EG
 {
 	LOG_SCOPE_VERBOSITY_OVERRIDE(LogEOSSDK, ELogVerbosity::NoLogging);
 
-	DestroyCurrentServiceModule();
-
 	FAccountId FirstAccountId, SecondAccountId;
 
 	FCreateSession::Params OpFirstCreateParams;
@@ -443,8 +435,6 @@ SESSIONS_TEST_CASE("Call FindSessions using custom setting as a filter Int", EG_
 {
 	LOG_SCOPE_VERBOSITY_OVERRIDE(LogEOSSDK, ELogVerbosity::NoLogging);
 
-	DestroyCurrentServiceModule();
-
 	FAccountId FirstAccountId, SecondAccountId;
 
 	FCreateSession::Params OpFirstCreateParams;
@@ -549,8 +539,6 @@ SESSIONS_TEST_CASE("Call FindSessions using custom setting as a filter Double",
 {
 	LOG_SCOPE_VERBOSITY_OVERRIDE(LogEOSSDK, ELogVerbosity::NoLogging);
 
-	DestroyCurrentServiceModule();
-
 	FAccountId FirstAccountId, SecondAccountId;
 
 	FCreateSession::Params OpFirstCreateParams;
@@ -687,8 +675,6 @@ SESSIONS_TEST_CASE("Call FindSessions using custom setting as a filter String",
 	LOG_SCOPE_VERBOSITY_OVERRIDE(LogOnlineServices, ELogVerbosity::NoLogging);
 	LOG_SCOPE_VERBOSITY_OVERRIDE(LogEOSSDK, ELogVerbosity::NoLogging);
 
-	DestroyCurrentServiceModule();
-
 	FAccountId FirstAccountId, SecondAccountId;
 
 	FCreateSession::Params OpFirstCreateSessionParams;
diff --git a/Engine/Source/Programs/Online/OnlineTestsCore/Source/Private/Tests/Sessions/GetSessionsTests.cpp b/Engine/Source/Programs/Online/OnlineTestsCore/Source/Private/Tests/Sessions/GetSessionsTests.cpp
--- a/Engine/Source/Programs/Online/OnlineTestsCore/Source/Private/Tests/Sessions/GetSessionsTests.cpp
+++ b/Engine/Source/Programs/Online/OnlineTestsCore/Source/Private/Tests/Sessions/GetSessionsTests.cpp
@@ -32,7 +32,6 @@ SESSIONS_TEST_CASE("If I call GetAllSessions with an invalid account id, I get a
 
 SESSIONS_TEST_CASE("If I call GetAllSessions before creating or joining any sessions, I get a successful result which is an empty array", EG_SESSIONS_GETSESSIONS_TAG)
 {
-	DestroyCurrentServiceModule();
 	ReturnAccounts();
 
 	FAccountId AccountId;
@@ -98,8 +97,6 @@ SESSIONS_TEST_CASE("If I call GetAllSessions with valid conditions, I get a vali
 
 SESSIONS_TEST_CASE("If I call GetSessionByName with an empty session name, I get an error", EG_SESSIONS_GETSESSIONS_TAG)
 {
-	DestroyCurrentServiceModule();
-
 	FAccountId AccountId;
 
 	GetLoginPipeline({ AccountId })
@@ -140,8 +137,6 @@ SESSIONS_TEST_CASE("If I call GetSessionByName with valid information, it return
 {
 	LOG_SCOPE_VERBOSITY_OVERRIDE(LogOnlineServices, ELogVerbosity::NoLogging);
 
-	DestroyCurrentServiceModule();
-
 	FAccountId AccountId;
 
 	FCreateSession::Params OpCreateParams;
@@ -181,8 +176,6 @@ SESSIONS_TEST_CASE("If I call GetSessionByName with valid information, it return
 
 SESSIONS_TEST_CASE("If I call GetSessionById with an invalid session id, I get an error", EG_SESSIONS_GETSESSIONS_TAG)
 {
-	DestroyCurrentServiceModule();
-
 	FAccountId AccountId;
 
 	GetLoginPipeline({ AccountId })
@@ -204,8 +197,6 @@ SESSIONS_TEST_CASE("If I call GetSessionById with a valid but unregistered sessi
 {
 	LOG_SCOPE_VERBOSITY_OVERRIDE(LogOnlineServices, ELogVerbosity::NoLogging);
 
-	DestroyCurrentServiceModule();
-
 	FAccountId AccountId;
 
 	FCreateSession::Params OpCreateParams;
@@ -338,8 +329,6 @@ SESSIONS_TEST_CASE("If I call GetPresenceSession with an unregistered id, I get
 
 SESSIONS_TEST_CASE("If I call GetPresenceSession with a valid id, I get a valid reference to the session", EG_SESSIONS_GETSESSIONS_TAG)
 {
-	DestroyCurrentServiceModule();
-
 	FAccountId AccountId;
 
 	FCreateSession::Params OpCreateParams;
@@ -423,8 +412,6 @@ SESSIONS_TEST_CASE("If I call GetSessionInviteById with an invalid session invit
 
 SESSIONS_TEST_CASE("If I call GetSessionInviteById with a valid account id, but without an existing invite, I get an error", EG_SESSIONS_GETSESSIONSEOS_TAG)
 {
-	DestroyCurrentServiceModule();
-
 	int32 TestAccountIndex = 7;
 	FAccountId FirstAccountId, SecondAccountId;
 
@@ -482,8 +469,6 @@ SESSIONS_TEST_CASE("If I call GetSessionInviteById with valid invite id, but wit
 {
 	LOG_SCOPE_VERBOSITY_OVERRIDE(LogEOSSDK, ELogVerbosity::NoLogging);
 
-	DestroyCurrentServiceModule();
-
 	int32 TestAccountIndex = 7;
 	FAccountId FirstAccountId, SecondAccountId;
 
@@ -546,8 +531,6 @@ SESSIONS_TEST_CASE("If I call GetSessionInviteById with valid invite id, but wit
 
 SESSIONS_TEST_CASE("If I call GetSessionInviteById with a valid data, I get a valid reference to the session invite", EG_SESSIONS_GETSESSIONSEOS_TAG)
 {
-	DestroyCurrentServiceModule();
-
 	int32 TestAccountIndex = 7;
 	FAccountId FirstAccountId, SecondAccountId;
 
diff --git a/Engine/Source/Programs/Online/OnlineTestsCore/Source/Private/Tests/Sessions/JoinSessionTests.cpp b/Engine/Source/Programs/Online/OnlineTestsCore/Source/Private/Tests/Sessions/JoinSessionTests.cpp
--- a/Engine/Source/Programs/Online/OnlineTestsCore/Source/Private/Tests/Sessions/JoinSessionTests.cpp
+++ b/Engine/Source/Programs/Online/OnlineTestsCore/Source/Private/Tests/Sessions/JoinSessionTests.cpp
@@ -55,8 +55,6 @@ SESSIONS_TEST_CASE("If I call JoinSession with an empty session name, I get an e
 {
 	LOG_SCOPE_VERBOSITY_OVERRIDE(LogEOSSDK, ELogVerbosity::NoLogging);
 
-	DestroyCurrentServiceModule();
-
 	int32 TestAccountIndex = 7;
 	FAccountId AccountId;
 
@@ -106,8 +104,6 @@ SESSIONS_TEST_CASE("If I call JoinSession with an name already in use, I get an
 {
 	LOG_SCOPE_VERBOSITY_OVERRIDE(LogEOSSDK, ELogVerbosity::NoLogging);
 
-	DestroyCurrentServiceModule();
-
 	int32 TestAccountIndex = 7;
 	FAccountId FirstAccountId, SecondAccountId;
 
@@ -227,8 +223,6 @@ SESSIONS_TEST_CASE("If I call JoinSession with a valid but unregistered session
 	LOG_SCOPE_VERBOSITY_OVERRIDE(LogEOSSDK, ELogVerbosity::NoLogging);
 	LOG_SCOPE_VERBOSITY_OVERRIDE(LogOnlineServices, ELogVerbosity::NoLogging);
 
-	DestroyCurrentServiceModule();
-
 	int32 TestAccountIndex = 7;
 	FAccountId FirstAccountId, SecondAccountId;
 
@@ -318,8 +312,6 @@ SESSIONS_TEST_CASE("If I call JoinSession with a valid but unregistered session
 
 SESSIONS_TEST_CASE("If I call JoinSession for a session I'm already registered in, I get an error", EG_SESSIONS_JOINSESSIONSEOS_TAG)
 {
-	DestroyCurrentServiceModule();
-
 	FAccountId FirstAccountId, SecondAccountId;
 
 	FCreateSession::Params OpCreateParams;
@@ -395,8 +387,6 @@ SESSIONS_TEST_CASE("If I call JoinSession for a session I'm already registered i
 
 SESSIONS_TEST_CASE("If I call JoinSession for a session that is not joinable, I get an error", EG_SESSIONS_JOINSESSIONSEOS_TAG)
 {
-	DestroyCurrentServiceModule();
-
 	int32 TestAccountIndex = 7;
 	FAccountId FirstAccountId, SecondAccountId;
 
@@ -466,8 +456,6 @@ SESSIONS_TEST_CASE("If I call JoinSession for a session that is not joinable, I
 
 SESSIONS_TEST_CASE("If I call JoinSession with presence enabled when there is already a presence session, I get an error", EG_SESSIONS_JOINSESSIONS_TAG)
 {
-	DestroyCurrentServiceModule();
-
 	FAccountId AccountId;
 
 	FCreateSession::Params OpCreateParams;
@@ -516,8 +504,6 @@ SESSIONS_TEST_CASE("If I call JoinSession with presence enabled when there is al
 
 SESSIONS_TEST_CASE("If I call JoinSession with valid information, the operation completes successfully", EG_SESSIONS_JOINSESSIONSEOS_TAG)
 {
-	DestroyCurrentServiceModule();
-
 	FAccountId FirstAccountId, SecondAccountId;
 
 	FCreateSession::Params OpCreateParams;
diff --git a/Engine/Source/Programs/Online/OnlineTestsCore/Source/Private/Tests/Sessions/LeaveSessionTests.cpp b/Engine/Source/Programs/Online/OnlineTestsCore/Source/Private/Tests/Sessions/LeaveSessionTests.cpp
--- a/Engine/Source/Programs/Online/OnlineTestsCore/Source/Private/Tests/Sessions/LeaveSessionTests.cpp
+++ b/Engine/Source/Programs/Online/OnlineTestsCore/Source/Private/Tests/Sessions/LeaveSessionTests.cpp
@@ -63,8 +63,6 @@ SESSIONS_TEST_CASE("If I call LeaveSession with an unregistered session name, I
 
 SESSIONS_TEST_CASE("If I call LeaveSession with valid data, the operation completes successfully", SESSIONS_LEAVESESSION_TAG)
 {
-	DestroyCurrentServiceModule();
-
 	int32 TestAccountIndex = 7;
 	FAccountId AccountId;
 
diff --git a/Engine/Source/Programs/Online/OnlineTestsCore/Source/Private/Tests/Sessions/PresenceSessionTests.cpp b/Engine/Source/Programs/Online/OnlineTestsCore/Source/Private/Tests/Sessions/PresenceSessionTests.cpp
--- a/Engine/Source/Programs/Online/OnlineTestsCore/Source/Private/Tests/Sessions/PresenceSessionTests.cpp
+++ b/Engine/Source/Programs/Online/OnlineTestsCore/Source/Private/Tests/Sessions/PresenceSessionTests.cpp
@@ -15,8 +15,6 @@
 
 SESSIONS_TEST_CASE("If I call IsPresenceSession with an invalid account id, I get an error", EG_SESSIONS_PRESENCESESSION_TAG)
 {
-	DestroyCurrentServiceModule();
-
 	GetPipeline()
 		.EmplaceLambda([](const IOnlineServicesPtr& OnlineSubsystem)
 			{
@@ -56,8 +54,6 @@ SESSIONS_TEST_CASE("If I call IsPresenceSession with an unregistered account id,
 {
 	LOG_SCOPE_VERBOSITY_OVERRIDE(LogOnlineServices, ELogVerbosity::NoLogging);
 
-	DestroyCurrentServiceModule();
-
 	int32 TestAccountIndex = 7;
 	FAccountId AccountId;
 
diff --git a/Engine/Source/Programs/Online/OnlineTestsCore/Source/Private/Tests/Sessions/SendRejectSessionInvitesTests.cpp b/Engine/Source/Programs/Online/OnlineTestsCore/Source/Private/Tests/Sessions/SendRejectSessionInvitesTests.cpp
--- a/Engine/Source/Programs/Online/OnlineTestsCore/Source/Private/Tests/Sessions/SendRejectSessionInvitesTests.cpp
+++ b/Engine/Source/Programs/Online/OnlineTestsCore/Source/Private/Tests/Sessions/SendRejectSessionInvitesTests.cpp
@@ -15,8 +15,6 @@
 
 SESSIONS_TEST_CASE("If I call SendSessionInvite with an invalid account id, I get an error", EG_SESSIONS_SENDREJECTSESSIONINVITE_TAG)
 {
-	DestroyCurrentServiceModule();
-
 	FSendSessionInvite::Params OpSendInviteParams;
 	FSendSessionInviteHelper::FHelperParams SendSessionInviteHelperParams;
 	SendSessionInviteHelperParams.OpParams = &OpSendInviteParams;
@@ -73,8 +71,6 @@ SESSIONS_TEST_CASE("If I call SendSessionInvite with a valid session name but un
 
 SESSIONS_TEST_CASE("If I call SendSessionInvite with an empty target users, I get an error", EG_SESSIONS_SENDREJECTSESSIONINVITE_MULTIACCOUNT_TAG)
 {
-	DestroyCurrentServiceModule();
-
 	int32 TestAccountIndex = 7;
 	FAccountId FirstAccountId, SecondAccountId;
 
@@ -114,8 +110,6 @@ SESSIONS_TEST_CASE("If I call SendSessionInvite with an empty target users, I ge
 
 SESSIONS_TEST_CASE("If I call SendSessionInvite with an invalid target users, I get an error", EG_SESSIONS_SENDREJECTSESSIONINVITE_MULTIACCOUNT_TAG)
 {
-	DestroyCurrentServiceModule();
-
 	int32 TestAccountIndex = 7;
 	FAccountId FirstAccountId, SecondAccountId;
 		
@@ -156,8 +150,6 @@ SESSIONS_TEST_CASE("If I call SendSessionInvite with an invalid target users, I
 
 SESSIONS_TEST_CASE("If I call SendSessionInvite with valid data, the operation completes successfully", EG_SESSIONS_SENDREJECTSESSIONINVITEEOS_TAG)
 {
-	DestroyCurrentServiceModule();
-
 	int32 TestAccountIndex = 7;
 	FAccountId FirstAccountId, SecondAccountId;
 
@@ -252,8 +244,6 @@ SESSIONS_TEST_CASE("If I call RejectSessionInvite with valid data, the operation
 {
 	LOG_SCOPE_VERBOSITY_OVERRIDE(LogEOSSDK, ELogVerbosity::NoLogging);
 
-	DestroyCurrentServiceModule();
-
 	int32 TestAccountIndex = 7;
 	FAccountId FirstAccountId, SecondAccountId;
 
diff --git a/Engine/Source/Programs/Online/OnlineTestsCore/Source/Private/Tests/Sessions/StartMatchmakingTests.cpp b/Engine/Source/Programs/Online/OnlineTestsCore/Source/Private/Tests/Sessions/StartMatchmakingTests.cpp
--- a/Engine/Source/Programs/Online/OnlineTestsCore/Source/Private/Tests/Sessions/StartMatchmakingTests.cpp
+++ b/Engine/Source/Programs/Online/OnlineTestsCore/Source/Private/Tests/Sessions/StartMatchmakingTests.cpp
@@ -139,8 +139,6 @@ SESSIONS_TEST_CASE("If I call StartMatchmaking with session override id less tha
 {
 	LOG_SCOPE_VERBOSITY_OVERRIDE(LogOnlineServices, ELogVerbosity::NoLogging);
 
-	DestroyCurrentServiceModule();
-
 	FAccountId AccountId;
 
 	FStartMatchmaking::Params OpFirstStartMatchmakingParams;
@@ -175,8 +173,6 @@ SESSIONS_TEST_CASE("If I call StartMatchmaking with session override id less tha
 
 SESSIONS_TEST_CASE("If I call StartMatchmaking with name already in use, I get an error", EG_SESSIONS_STARTMATCHMAKING_TAG)
 {
-	DestroyCurrentServiceModule();
-
 	int32 TestAccountIndex = 7;
 	FAccountId AccountId;
 
diff --git a/Engine/Source/Programs/Online/OnlineTestsCore/Source/Private/Tests/Sessions/UpdateSessionSettingsTests.cpp b/Engine/Source/Programs/Online/OnlineTestsCore/Source/Private/Tests/Sessions/UpdateSessionSettingsTests.cpp
--- a/Engine/Source/Programs/Online/OnlineTestsCore/Source/Private/Tests/Sessions/UpdateSessionSettingsTests.cpp
+++ b/Engine/Source/Programs/Online/OnlineTestsCore/Source/Private/Tests/Sessions/UpdateSessionSettingsTests.cpp
@@ -128,8 +128,6 @@ UPDATESESSIONS_TEST_CASE("If I call UpdateSessionSettings with an unregistered s
 
 UPDATESESSIONS_TEST_CASE("If I call UpdateSessionSettings with valid data, the operation completes successfully", SESSIONS_UPDATESESSIONSETTINGS_TAG)
 {
-	DestroyCurrentServiceModule();
-
 	int32 TestAccountIndex = 7;
 	FAccountId AccountId;
 
diff --git a/Engine/Source/Programs/Online/OnlineTestsCore/Source/Private/Tests/TestAsyncOp.cpp b/Engine/Source/Programs/Online/OnlineTestsCore/Source/Private/Tests/TestAsyncOp.cpp
--- a/Engine/Source/Programs/Online/OnlineTestsCore/Source/Private/Tests/TestAsyncOp.cpp
+++ b/Engine/Source/Programs/Online/OnlineTestsCore/Source/Private/Tests/TestAsyncOp.cpp
@@ -309,7 +309,6 @@ ASYNCOP_TEST_CASE("Verify that the duration cache policy works)")
 
 ASYNCOP_TEST_CASE("Verify that the cache error setting works)", ASYNCOP_SKIP_EOS_PLATFORM_FAULT_TAG)
 {
-	DestroyCurrentServiceModule();
 	FTestPipeline& TestPipeline = GetPipeline();
 	TArray<FString> DontCacheErrorHeirarchy = { TEXT("OnlineTests.TestDontCacheError") };
 	// Uses the following operation configuration in OnlineTests' DefaultEngine.ini
diff --git a/Engine/Source/Programs/Online/OnlineTestsCore/Source/Private/Tests/TitleFile/TitleFileTests.cpp b/Engine/Source/Programs/Online/OnlineTestsCore/Source/Private/Tests/TitleFile/TitleFileTests.cpp
--- a/Engine/Source/Programs/Online/OnlineTestsCore/Source/Private/Tests/TitleFile/TitleFileTests.cpp
+++ b/Engine/Source/Programs/Online/OnlineTestsCore/Source/Private/Tests/TitleFile/TitleFileTests.cpp
@@ -90,8 +90,6 @@ TITLE_FILE_TEST_CASE("GetEnumeratedFiles - Invalid user", TITLE_FILE_GET_ENUMERA
 
 TITLE_FILE_TEST_CASE("GetEnumeratedFiles - Invalid state", TITLE_FILE_GET_ENUMERATED_SKIP_EOS_PLATFORM_FAULT_TAG)
 {
-	DestroyCurrentServiceModule();
-
 	FAccountId AccountId;
 	FTestPipeline& Pipe = GetLoginPipeline({ AccountId });
 
@@ -147,8 +145,6 @@ TITLE_FILE_TEST_CASE("ReadFile - Invalid User", TITLE_FILE_READ_TAG)
 
 TITLE_FILE_TEST_CASE("ReadFile - Empty filename", TITLE_FILE_READ_TAG)
 {
-	DestroyCurrentServiceModule();
-
 	FAccountId AccountId;
 	FTestPipeline& Pipe = GetLoginPipeline({ AccountId });
 
diff --git a/Engine/Source/Programs/Online/OnlineTestsCore/Source/Private/Tests/UserFile/UserFileTests.cpp b/Engine/Source/Programs/Online/OnlineTestsCore/Source/Private/Tests/UserFile/UserFileTests.cpp
--- a/Engine/Source/Programs/Online/OnlineTestsCore/Source/Private/Tests/UserFile/UserFileTests.cpp
+++ b/Engine/Source/Programs/Online/OnlineTestsCore/Source/Private/Tests/UserFile/UserFileTests.cpp
@@ -120,8 +120,6 @@ USER_FILE_TEST_CASE("GetEnumeratedFiles Tests - Invalid State", USER_FILE_GETENU
 	FScopeDisableWarningsInLog ScopeLogSDK{ &LogEOSSDK };
 	FScopeDisableWarningsInLog ScopeLogOnlineSDK{ &LogOnlineServices };
 
-	DestroyCurrentServiceModule();
-
 	FAccountId AccountId;
 	FTestPipeline& Pipe = GetLoginPipeline({ AccountId });
 
