Add an ensure to catch when TOnlineAsyncOps do not have a result/error set.
- In order to perform this check, there is a new requirement that a TOnlineAsyncOp result/error is set before the final step completes.
- Fix OnlineTests where a result was not set
- Update a case in SessionsCommon where the result was set after all steps were complete

diff --git a/Engine/Plugins/Online/OnlineServices/Source/OnlineServicesCommon/Private/Online/SessionsCommon.cpp b/Engine/Plugins/Online/OnlineServices/Source/OnlineServicesCommon/Private/Online/SessionsCommon.cpp
--- a/Engine/Plugins/Online/OnlineServices/Source/OnlineServicesCommon/Private/Online/SessionsCommon.cpp
+++ b/Engine/Plugins/Online/OnlineServices/Source/OnlineServicesCommon/Private/Online/SessionsCommon.cpp
@@ -372,20 +372,22 @@ namespace UE::Online {
 			return Op->GetHandle();
 		}
 
-		Op->Then([this, WeakThis = AsWeak(), ImplFunc](TOnlineAsyncOp<MethodStruct>& Op)
+		Op->Then([this, ImplFunc](TOnlineAsyncOp<MethodStruct>& Op)
 			{
-				if (TSharedPtr<ISessions> StrongThis = WeakThis.Pin())
-				{
-					const typename MethodStruct::Params& OpParams = Op.GetParams();
+				const typename MethodStruct::Params& OpParams = Op.GetParams();
 
-					if (TOptional<FOnlineError> StateCheck = CheckState(Op.GetParams()))
-					{
-						Op.SetError(MoveTemp(StateCheck.GetValue()));
-						return;
-					}
+				TPromise<void> Promise;
+				TFuture<void> Future = Promise.GetFuture();
 
+				if (TOptional<FOnlineError> StateCheck = CheckState(Op.GetParams()))
+				{
+					Op.SetError(MoveTemp(StateCheck.GetValue()));
+					Promise.EmplaceValue();
+				}
+				else
+				{
 					(this->*ImplFunc)(OpParams)
-						.Next([WeakOp = Op.AsWeak()](const TOnlineResult<MethodStruct>& Result)
+						.Next([WeakOp = Op.AsWeak(), MovedPromise = MoveTemp(Promise)](const TOnlineResult<MethodStruct>& Result) mutable
 							{
 								if (TOnlineAsyncOpPtr<MethodStruct> StrongOp = WeakOp.Pin())
 								{
@@ -398,8 +400,11 @@ namespace UE::Online {
 										StrongOp->SetError(FOnlineError(Result.GetErrorValue()));
 									}
 								}
+								MovedPromise.EmplaceValue();
 							});
 				}
+
+				return Future;
 			})
 			.Enqueue(bUseParallelQueue ? Services.GetParallelQueue() : GetSerialQueue());
 
@@ -736,7 +741,7 @@ namespace UE::Online {
 
 		FoundSession->SessionMembers.Emplace(Params.LocalAccountId);
 
-		return TOnlineResult<FAddSessionMember>();
+		return TOnlineResult<FAddSessionMember>(FAddSessionMember::Result());
 	}
 
 	TOptional<FOnlineError> FSessionsCommon::CheckParams(const FAddSessionMember::Params& Params) const
@@ -785,7 +790,7 @@ namespace UE::Online {
 
 		FoundSession->SessionMembers.Remove(Params.LocalAccountId);
 
-		return TOnlineResult<FRemoveSessionMember>();
+		return TOnlineResult<FRemoveSessionMember>(FRemoveSessionMember::Result());
 	}
 
 	TOptional<FOnlineError> FSessionsCommon::CheckParams(const FRemoveSessionMember::Params& Params) const
diff --git a/Engine/Plugins/Online/OnlineServices/Source/OnlineServicesCommon/Public/Online/OnlineAsyncOp.h b/Engine/Plugins/Online/OnlineServices/Source/OnlineServicesCommon/Public/Online/OnlineAsyncOp.h
--- a/Engine/Plugins/Online/OnlineServices/Source/OnlineServicesCommon/Public/Online/OnlineAsyncOp.h
+++ b/Engine/Plugins/Online/OnlineServices/Source/OnlineServicesCommon/Public/Online/OnlineAsyncOp.h
@@ -1477,7 +1477,7 @@ public:
 		{
 			const int32 StepToExecute = NextStep;
 			++NextStep;
-			if (StepToExecute < Steps.Num())
+			if (ensureMsgf(StepToExecute < Steps.Num(), TEXT("All OnlineAsyncOp steps completed but no error or result was set")))
 			{
 				Execute(Steps[StepToExecute]->GetExecutionPolicy(),
 					[this, StepToExecute, WeakThis = TWeakPtr<TOnlineAsyncOp<OpType>>(this->AsShared())]()
@@ -1496,7 +1496,7 @@ public:
 		if (!IsComplete())
 		{
 			const int32 StepToExecute = NextStep == 0 ? 0 : NextStep - 1;
-			if (StepToExecute < Steps.Num())
+			if (ensureMsgf(StepToExecute < Steps.Num(), TEXT("All OnlineAsyncOp steps completed but no error or result was set")))
 			{
 				Execute(Steps[StepToExecute]->GetExecutionPolicy(),
 					[this, StepToExecute, WeakThis = TWeakPtr<TOnlineAsyncOp<OpType>>(this->AsShared())]()
diff --git a/Engine/Source/Programs/Online/OnlineTestsCore/Source/Private/Tests/TestAsyncOp.cpp b/Engine/Source/Programs/Online/OnlineTestsCore/Source/Private/Tests/TestAsyncOp.cpp
--- a/Engine/Source/Programs/Online/OnlineTestsCore/Source/Private/Tests/TestAsyncOp.cpp
+++ b/Engine/Source/Programs/Online/OnlineTestsCore/Source/Private/Tests/TestAsyncOp.cpp
@@ -76,7 +76,11 @@ ASYNCOP_TEST_CASE("Verify that the cache returns the same joinable op when the p
 			Params1.Index = 0;
 			TOnlineAsyncOpRef<FTestOpWithParams> Op = ServicesCommon.GetJoinableOp<FTestOpWithParams>(MoveTemp(Params1));
 			CHECK(!Op->IsReady());
-			Op->Then([Promise](FOnlineAsyncOp&) { Promise->SetValue(true); })
+			Op->Then([Promise](TOnlineAsyncOp<FTestOpWithParams>& InOp)
+			{
+				InOp.SetResult(FTestOpWithParams::Result());
+				Promise->SetValue(true);
+			})
 			.Enqueue(ServicesCommon.GetParallelQueue());
 
 			FTestOpWithParams::Params Params2;
@@ -98,14 +102,21 @@ ASYNCOP_TEST_CASE("Verify that the cache returns different joinable ops when the
 			Params1.Index = 1;
 			TOnlineAsyncOpRef<FTestOpWithParams> Op = ServicesCommon.GetJoinableOp<FTestOpWithParams>(MoveTemp(Params1));
 			CHECK(!Op->IsReady());
-			Op->Then([](FOnlineAsyncOp&) {})
+			Op->Then([](TOnlineAsyncOp<FTestOpWithParams>& InOp)
+			{
+				InOp.SetResult(FTestOpWithParams::Result());
+			})
 			.Enqueue(ServicesCommon.GetParallelQueue());
 
 			FTestOpWithParams::Params Params2;
 			Params2.Index = 2;
 			TOnlineAsyncOpRef<FTestOpWithParams> Op2 = ServicesCommon.GetJoinableOp<FTestOpWithParams>(MoveTemp(Params2));
 			CHECK(!Op2->IsReady());
-			Op2->Then([Promise](FOnlineAsyncOp&) { Promise->SetValue(true); })
+			Op2->Then([Promise](TOnlineAsyncOp<FTestOpWithParams>& InOp)
+			{
+				InOp.SetResult(FTestOpWithParams::Result());
+				Promise->SetValue(true);
+			})
 			.Enqueue(ServicesCommon.GetParallelQueue());
 		});
 
@@ -122,7 +133,11 @@ ASYNCOP_TEST_CASE("Verify that the cache returns the same mergeable op when the
 			Params1.Index = 0;
 			TOnlineAsyncOpRef<FMergeableTestOpWithParams> Op = ServicesCommon.GetMergeableOp<FMergeableTestOpWithParams>(MoveTemp(Params1));
 			CHECK(!Op->IsReady());
-			Op->Then([Promise](FOnlineAsyncOp&) { Promise->SetValue(true); })
+			Op->Then([Promise](TOnlineAsyncOp<FMergeableTestOpWithParams>& InOp)
+			{
+				InOp.SetResult(FMergeableTestOpWithParams::Result());
+				Promise->SetValue(true);
+			})
 			.Enqueue(ServicesCommon.GetParallelQueue());
 
 			FMergeableTestOpWithParams::Params Params2;
@@ -144,14 +159,21 @@ ASYNCOP_TEST_CASE("Verify that the cache returns different mergeable ops when th
 			Params1.Index = 1;
 			TOnlineAsyncOpRef<FMergeableTestOpWithParams> Op = ServicesCommon.GetMergeableOp<FMergeableTestOpWithParams>(MoveTemp(Params1));
 			CHECK(!Op->IsReady());
-			Op->Then([](FOnlineAsyncOp&) {})
+			Op->Then([](TOnlineAsyncOp<FMergeableTestOpWithParams>& InOp)
+			{
+				InOp.SetResult(FMergeableTestOpWithParams::Result());
+			})
 			.Enqueue(ServicesCommon.GetParallelQueue());
 
 			FMergeableTestOpWithParams::Params Params2;
 			Params2.Index = 2;
 			TOnlineAsyncOpRef<FMergeableTestOpWithParams> Op2 = ServicesCommon.GetMergeableOp<FMergeableTestOpWithParams>(MoveTemp(Params2));
 			CHECK(!Op2->IsReady());
-			Op2->Then([Promise](FOnlineAsyncOp&) { Promise->SetValue(true); })
+			Op2->Then([Promise](TOnlineAsyncOp<FMergeableTestOpWithParams>& InOp)
+			{
+				InOp.SetResult(FMergeableTestOpWithParams::Result());
+				Promise->SetValue(true);
+			})
 			.Enqueue(ServicesCommon.GetParallelQueue());
 		});
 
