Compiler fixes for tests

#rnx

diff --git a/Engine/Source/Developer/LowLevelTestsRunner/Public/TestCommon/Initialization.h b/Engine/Source/Developer/LowLevelTestsRunner/Public/TestCommon/Initialization.h
--- a/Engine/Source/Developer/LowLevelTestsRunner/Public/TestCommon/Initialization.h
+++ b/Engine/Source/Developer/LowLevelTestsRunner/Public/TestCommon/Initialization.h
@@ -2,8 +2,6 @@
 
 #pragma once
 
-#include "CoreMinimal.h"
-
 void SaveDefaultPlatformFile();
 void UsePlatformFileStubIfRequired();
 void UseDefaultPlatformFile();
diff --git a/Engine/Source/Programs/Online/OSSTestsCore/Tests/Purchase/PurchaseGetReceiptsTests.cpp b/Engine/Source/Programs/Online/OSSTestsCore/Tests/Purchase/PurchaseGetReceiptsTests.cpp
--- a/Engine/Source/Programs/Online/OSSTestsCore/Tests/Purchase/PurchaseGetReceiptsTests.cpp
+++ b/Engine/Source/Programs/Online/OSSTestsCore/Tests/Purchase/PurchaseGetReceiptsTests.cpp
@@ -27,7 +27,6 @@ PURCHASE_TEST_CASE("Verify calling GetReceipts with valid inputs returns the exp
 	bool bRestoreReceipts = false;
 	bool bIsShouldBeEmpty = false;
 
-	QueryOwnership
 	int32 NumUsersToImplicitLogin = 1;
 	TArray<FPurchaseReceipt>* PurchaseReceipts = nullptr;
 
