OSSTests: Fixed compile error: Cannot open include file: 'Helpers/OnlineSubsystemTestFixtures.h'

diff --git a/Engine/Source/Programs/Online/OSSTestsCore/Tests/Achievements/AchievementsGetCachedAchievementDescriptionTests.cpp b/Engine/Source/Programs/Online/OSSTestsCore/Tests/Achievements/AchievementsGetCachedAchievementDescriptionTests.cpp
--- a/Engine/Source/Programs/Online/OSSTestsCore/Tests/Achievements/AchievementsGetCachedAchievementDescriptionTests.cpp
+++ b/Engine/Source/Programs/Online/OSSTestsCore/Tests/Achievements/AchievementsGetCachedAchievementDescriptionTests.cpp
@@ -11,8 +11,6 @@
 
 #include "Helpers/TickForTime.h"
 
-#include "Helpers/OnlineSubsystemTestFixtures.h"
-
 #define ACHIEVEMENTS_TAG "[suite_achievements]"
 #define EG_ACHIEVEMENTS_GETCACHEDACHIEMENTDESCRIPTION_TAG ACHIEVEMENTS_TAG "[getcachedachievementdescription]"
 
diff --git a/Engine/Source/Programs/Online/OSSTestsCore/Tests/Identity/IdentityAutoLoginTests.cpp b/Engine/Source/Programs/Online/OSSTestsCore/Tests/Identity/IdentityAutoLoginTests.cpp
--- a/Engine/Source/Programs/Online/OSSTestsCore/Tests/Identity/IdentityAutoLoginTests.cpp
+++ b/Engine/Source/Programs/Online/OSSTestsCore/Tests/Identity/IdentityAutoLoginTests.cpp
@@ -4,7 +4,6 @@
 
 #include "Helpers/Identity/IdentityAutoLoginHelper.h"
 #include "Helpers/Identity/IdentityLogoutHelper.h"
-#include "Helpers/OnlineSubsystemTestFixtures.h"
 
 #include "Misc/CommandLine.h"
 #include "OnlineSubsystemCatchHelper.h"
@@ -15,9 +14,7 @@
 
 #define IDENTITY_TEST_CASE(x, ...) ONLINESUBSYSTEM_TEST_CASE(x, IDENTITY_TAG __VA_ARGS__)
 
-#define IDENTITY_LEGACY_AUTOLOGIN_EOS_TEST(x, ...) ONLINESUBSYSTEM_TEST_CASE_FIXTURE(FOnlineSubsystemEOSLegacyTestFixture, x, __VA_ARGS__)
-
-IDENTITY_LEGACY_AUTOLOGIN_EOS_TEST("Verify calling Identity AutoLogin with valid inputs and using legacy login flow returns the expected result(Success Case)", EG_IDENTITY_AUTOLOGIN_LEGACY_EOS_FLOW_TAG)
+IDENTITY_TEST_CASE("Verify calling Identity AutoLogin with valid inputs and using legacy login flow returns the expected result(Success Case)", EG_IDENTITY_AUTOLOGIN_LEGACY_EOS_FLOW_TAG)
 {
 	FTestDriver LocalDriver;
 	int32 LocalUserNum = 0;
diff --git a/Engine/Source/Programs/Online/OSSTestsCore/Tests/Identity/IdentityLoginTests.cpp b/Engine/Source/Programs/Online/OSSTestsCore/Tests/Identity/IdentityLoginTests.cpp
--- a/Engine/Source/Programs/Online/OSSTestsCore/Tests/Identity/IdentityLoginTests.cpp
+++ b/Engine/Source/Programs/Online/OSSTestsCore/Tests/Identity/IdentityLoginTests.cpp
@@ -3,7 +3,6 @@
 #include "TestDriver.h"
 
 #include "OnlineSubsystemCatchHelper.h"
-#include "Helpers/OnlineSubsystemTestFixtures.h"
 
 #define IDENTITY_TAG "[suite_identity]"
 #define EG_IDENTITY_LOGIN_TAG IDENTITY_TAG "[login]"
@@ -11,9 +10,7 @@
 
 #define IDENTITY_TEST_CASE(x, ...) ONLINESUBSYSTEM_TEST_CASE(x, IDENTITY_TAG __VA_ARGS__)
 
-#define IDENTITY_LEGACY_LOGIN_EOS_TEST(x, ...) ONLINESUBSYSTEM_TEST_CASE_FIXTURE(FOnlineSubsystemEOSLegacyTestFixture, x, IDENTITY_TAG __VA_ARGS__)
-
-IDENTITY_LEGACY_LOGIN_EOS_TEST("Verify calling Identity Login with valid inputs and using legacy login flow returns the expected result(Success Case)", EG_IDENTITY_LOGIN_LEGACY_FLOWEOS_TAG)
+IDENTITY_TEST_CASE("Verify calling Identity Login with valid inputs and using legacy login flow returns the expected result(Success Case)", EG_IDENTITY_LOGIN_LEGACY_FLOWEOS_TAG)
 {
 	int32 NumUsersToImplicitLogin = 2;
 
