- Migrate some OSSv2 helpers from FNOnlineFrameworkTests to OnlineTestsCore so more test targets can use them
- Remove OnlineServices EOS and Null as dependencies
- Add missing ONLINETESTS_USEEXTERNAUTH define checks around calls to ReturnAccounts

#rnx

diff --git a/Engine/Source/Programs/Online/OnlineTestsCore/Source/OnlineTestsCore.Build.cs b/Engine/Source/Programs/Online/OnlineTestsCore/Source/OnlineTestsCore.Build.cs
--- a/Engine/Source/Programs/Online/OnlineTestsCore/Source/OnlineTestsCore.Build.cs
+++ b/Engine/Source/Programs/Online/OnlineTestsCore/Source/OnlineTestsCore.Build.cs
@@ -25,9 +25,7 @@ public class OnlineTestsCore : ModuleRules
 				"EOSShared",
 				"OnlineSubsystem",
 				"OnlineServicesInterface",
-				"OnlineServicesCommon",
-				"OnlineServicesEOS",
-				"OnlineServicesNull",
+				"OnlineServicesCommon",				
 				"OnlineServicesOSSAdapter",
 				"SSL",
 				"Json",
diff --git a/Engine/Source/Programs/Online/OnlineTestsCore/Source/Private/Tests/Sessions/FindSessionsTests.cpp b/Engine/Source/Programs/Online/OnlineTestsCore/Source/Private/Tests/Sessions/FindSessionsTests.cpp
--- a/Engine/Source/Programs/Online/OnlineTestsCore/Source/Private/Tests/Sessions/FindSessionsTests.cpp
+++ b/Engine/Source/Programs/Online/OnlineTestsCore/Source/Private/Tests/Sessions/FindSessionsTests.cpp
@@ -135,7 +135,9 @@ SESSIONS_TEST_CASE("If I call FindSessions with an empty search filter key, I ge
 
 SESSIONS_TEST_CASE("If I call FindSessions while a search is already running for that user, I get an error", EG_SESSIONS_FINDSESSIONS_TAG)
 {
+#if ONLINETESTS_USEEXTERNAUTH
 	ReturnAccounts();
+#endif
 
 	int32 TestAccountIndex = 7;
 	FAccountId FirstAccountId, SecondAccountId;
@@ -232,7 +234,9 @@ SESSIONS_TEST_CASE("If I call FindSessions with valid information, I can get an
 
 SESSIONS_TEST_CASE("If I call FindSessions with valid information, I will still not find sessions which JoinPolicy is not public", EG_SESSIONS_FINDSESSIONS_TAG)
 {
+#if ONLINETESTS_USEEXTERNAUTH
 	ReturnAccounts();
+#endif
 	
 	FAccountId FirstAccountId, SecondAccountId;
 	
diff --git a/Engine/Source/Programs/Online/OnlineTestsCore/Source/Private/Tests/Sessions/GetSessionsTests.cpp b/Engine/Source/Programs/Online/OnlineTestsCore/Source/Private/Tests/Sessions/GetSessionsTests.cpp
--- a/Engine/Source/Programs/Online/OnlineTestsCore/Source/Private/Tests/Sessions/GetSessionsTests.cpp
+++ b/Engine/Source/Programs/Online/OnlineTestsCore/Source/Private/Tests/Sessions/GetSessionsTests.cpp
@@ -32,7 +32,9 @@ SESSIONS_TEST_CASE("If I call GetAllSessions with an invalid account id, I get a
 
 SESSIONS_TEST_CASE("If I call GetAllSessions before creating or joining any sessions, I get a successful result which is an empty array", EG_SESSIONS_GETSESSIONS_TAG)
 {
+#if ONLINETESTS_USEEXTERNAUTH
 	ReturnAccounts();
+#endif
 
 	FAccountId AccountId;
 	GetLoginPipeline({ AccountId })
diff --git a/Engine/Source/Programs/Online/OnlineTestsCore/Source/Public/Helpers/OnlineHelpers.h b/Engine/Source/Programs/Online/OnlineTestsCore/Source/Public/Helpers/OnlineHelpers.h
new file mode 100644
--- /dev/null
+++ b/Engine/Source/Programs/Online/OnlineTestsCore/Source/Public/Helpers/OnlineHelpers.h
@@ -0,0 +1,163 @@
+// Copyright Epic Games, Inc. All Rights Reserved.
+ 
+#pragma once
+
+#include "Algo/Count.h"
+#include "Async/Future.h"
+#include "Online/OnlineAsyncOpHandle.h"
+#include "Online/OnlineError.h"
+#include "Online/OnlineResult.h"
+#include "Templates/ValueOrError.h"
+#include "TestHarness.h"
+
+namespace UE::TestCommon {
+
+constexpr const int32 TickFrequencyMs = (1000 / 30); // 30 ticks per second
+constexpr const int32 BlockUntilCompleteDefaultMaxTicks = (1000 * 30) / TickFrequencyMs; // Default timeout if 30 seconds
+
+/**
+ * Tick core systems that cause functionality to progress (eg core ticker)
+ */
+void Tick();
+
+/** Sleep behavior for BlockUntil Complete */
+enum class ESleepBehavior : uint8
+{
+	/** No sleeps */
+	NoSleep,
+	/** Sleep every tick */
+	Sleep
+};
+
+/**
+ * Tick core systems until a future is set, or the operation has timed out.
+ * Note: Unfulfilled promises are disastrous, so this fails the test if a timeout occurs.
+ * @param Futures futures to wait for
+ * @param MaxTicks maximum number of ticks before deciding the operation timed out.
+ * @param SleepBehavior sleep behavior while ticking
+ */
+template<typename ResultType>
+void BlockUntilComplete(TArrayView<const TFuture<ResultType>> Futures, TOptional<int32> MaxTicks = BlockUntilCompleteDefaultMaxTicks, ESleepBehavior SleepBehavior = ESleepBehavior::NoSleep)
+{
+	// Tick engine systems until our future is set, or we hit a timeout limit
+	int32 TicksRemaining = MaxTicks.Get(TNumericLimits<int32>::Max());
+
+	while ((Algo::CountIf(Futures, &TFuture<ResultType>::IsValid) - Algo::CountIf(Futures, &TFuture<ResultType>::IsReady)) > 0)
+	{
+		Tick();
+		--TicksRemaining;
+		if (TicksRemaining <= 0)
+		{
+			REQUIRE(false);
+		}
+		if (SleepBehavior == ESleepBehavior::Sleep)
+		{
+			// Some tests take actual wall time, not just ticks to pump data through queues, such as compressing data.
+			FPlatformProcess::Sleep(0.001f);
+		}
+	}
+}
+
+// Overload that takes a TArray of non-const futures
+template<typename ResultType>
+void BlockUntilComplete(const TArray<TFuture<ResultType>>& Futures, TOptional<int32> MaxTicks = BlockUntilCompleteDefaultMaxTicks, ESleepBehavior SleepBehavior = ESleepBehavior::NoSleep)
+{
+	BlockUntilComplete(MakeArrayView(const_cast<const TFuture<ResultType>*>(Futures.GetData()), Futures.Num()), BlockUntilCompleteDefaultMaxTicks, SleepBehavior);
+}
+
+// Overload that takes a single future
+template<typename ResultType>
+void BlockUntilComplete(const TFuture<ResultType>& Future, TOptional<int32> MaxTicks = BlockUntilCompleteDefaultMaxTicks, ESleepBehavior SleepBehavior = ESleepBehavior::NoSleep)
+{
+	BlockUntilComplete(MakeArrayView(&Future, 1), BlockUntilCompleteDefaultMaxTicks, SleepBehavior);
+}
+
+/** GetErrorFromResult - used by FulfillPromise to determine which error to fulfill a promise with */
+// Overload that takes TOnlineResult
+template<typename OpType>
+[[nodiscard]] UE::Online::FOnlineError GetErrorFromResult(UE::Online::TOnlineResult<OpType>&& InResult)
+{
+	return InResult.IsOk() ? UE::Online::Errors::Success() : MoveTemp(InResult.GetErrorValue());
+}
+
+// Overload that takes TValueOrError where Error is FOnlineError
+template<typename SuccessType>
+[[nodiscard]] UE::Online::FOnlineError GetErrorFromResult(TValueOrError<SuccessType, UE::Online::FOnlineError>&& InResult)
+{
+	return InResult.HasValue() ? UE::Online::Errors::Success() : InResult.StealError();
+}
+
+// Overload that takes FOnlineError
+[[nodiscard]] inline UE::Online::FOnlineError GetErrorFromResult(UE::Online::FOnlineError&& InError)
+{
+	return MoveTemp(InError);
+}
+
+
+/**
+ * Fulfill a promise of type FOnlineError
+ */
+template<typename TResultType>
+void FulfillPromise(TPromise<UE::Online::FOnlineError>& InPromise, TResultType&& InResult)
+{
+	InPromise.SetValue(GetErrorFromResult(MoveTempIfPossible(InResult)));
+}
+
+/**
+ * Check that a synchronous online operation completed successfully, and return the successful result.
+ * @param InResult the result
+ * @return the result of the op
+ */
+template<typename OpType>
+[[nodiscard]] OpType::Result GetOpResultChecked(UE::Online::TOnlineResult<OpType>&& InResult)
+{
+	REQUIRE(InResult.IsOk());
+	return InResult.GetOkValue();
+}
+
+/**
+ * Block until an async op completes.
+ * Will fail if the async op does not complete in a timely manner.
+ * @param InAsyncOp the async op to wait for
+ * @return the result of the async op
+ */
+template<typename OpType>
+[[nodiscard]] UE::Online::TOnlineResult<OpType> BlockUntilAsyncOpComplete(UE::Online::TOnlineAsyncOpHandle<OpType> InAsyncOp, TOptional<int32> MaxTicks = BlockUntilCompleteDefaultMaxTicks, ESleepBehavior SleepBehavior = ESleepBehavior::NoSleep)
+{
+	using namespace UE::Online;
+	TSharedRef<TPromise<TOnlineResult<OpType>>> Promise = MakeShared<TPromise<TOnlineResult<OpType>>>();
+	TFuture<TOnlineResult<OpType>> Future = Promise->GetFuture();
+	InAsyncOp.OnComplete([Promise](const TOnlineResult<OpType>& Result)
+	{
+		Promise->EmplaceValue(Result);
+	});
+	BlockUntilComplete(Future, MaxTicks, SleepBehavior);
+	return Future.Get();
+}
+
+/**
+ * Block until an async op completes, and return the result as an FOnlineError.
+ * Will return Errors::Success on success.
+ * Will fail if the async op does not complete in a timely manner.
+ * @param InAsyncOp the async op to wait for
+ * @return the FOnlineError result of the async op. Will be Errors::Success if the op was successful.
+ */
+template<typename OpType>
+[[nodiscard]] UE::Online::FOnlineError GetAsyncOpResultBasic(UE::Online::TOnlineAsyncOpHandle<OpType> InAsyncOp)
+{
+	return GetErrorFromResult(BlockUntilAsyncOpComplete(InAsyncOp));
+}
+
+/**
+ * Block until an async op completes and returns the successful variant of the result.
+ * Will fail if the async op fails, or the async op does not complete in a timely manner.
+ * @param InAsyncOp the async op to wait for
+ * @return the result of the async op
+ */
+template<typename OpType>
+[[nodiscard]] OpType::Result GetAsyncOpResultChecked(UE::Online::TOnlineAsyncOpHandle<OpType> InAsyncOp, TOptional<int32> MaxTicks = BlockUntilCompleteDefaultMaxTicks, ESleepBehavior SleepBehavior = ESleepBehavior::NoSleep)
+{
+	return GetOpResultChecked(BlockUntilAsyncOpComplete(InAsyncOp, MaxTicks, SleepBehavior));
+}
+
+/* UE::TestCommon */ }
\ No newline at end of file
