Add missing UE::TestCommon::Tick function

diff --git a/Engine/Source/Programs/Online/OnlineTestsCore/Source/Private/OnlineHelpers.cpp b/Engine/Source/Programs/Online/OnlineTestsCore/Source/Private/OnlineHelpers.cpp
new file mode 100644
--- /dev/null
+++ b/Engine/Source/Programs/Online/OnlineTestsCore/Source/Private/OnlineHelpers.cpp
@@ -0,0 +1,17 @@
+// Copyright Epic Games, Inc. All Rights Reserved.
+ 
+#include "Helpers/OnlineHelpers.h"
+
+#include "Async/TaskGraphInterfaces.h"
+#include "Containers/Ticker.h"
+
+namespace UE::TestCommon {
+
+void Tick()
+{
+	double TickTime = FTimespan::FromMilliseconds(TickFrequencyMs).GetTotalSeconds();
+	FTSTicker::GetCoreTicker().Tick(TickTime);
+	FTaskGraphInterface::Get().ProcessThreadUntilIdle(ENamedThreads::GameThread);
+}
+
+/* UE::TestCommon */ }
\ No newline at end of file
