Rename wait-for-future style functions

diff --git a/Engine/Source/Programs/Online/OnlineTestsCore/Source/Public/Helpers/OnlineHelpers.h b/Engine/Source/Programs/Online/OnlineTestsCore/Source/Public/Helpers/OnlineHelpers.h
--- a/Engine/Source/Programs/Online/OnlineTestsCore/Source/Public/Helpers/OnlineHelpers.h
+++ b/Engine/Source/Programs/Online/OnlineTestsCore/Source/Public/Helpers/OnlineHelpers.h
@@ -72,6 +72,14 @@ void BlockUntilComplete(const TFuture<ResultType>& Future, TOptional<int32> MaxT
 	BlockUntilComplete(MakeArrayView(&Future, 1), BlockUntilCompleteDefaultMaxTicks, SleepBehavior);
 }
 
+/** Block until a future is complete, and return its value */
+template <typename ResultType>
+ResultType GetFutureValueBlocking(TFuture<ResultType>& InFuture, TOptional<int32> MaxTicks = BlockUntilCompleteDefaultMaxTicks, ESleepBehavior SleepBehavior = ESleepBehavior::NoSleep)
+{
+	BlockUntilComplete(InFuture, MaxTicks, SleepBehavior);
+	return InFuture.Get();
+}
+
 /** GetErrorFromResult - used by FulfillPromise to determine which error to fulfill a promise with */
 // Overload that takes TOnlineResult
 template<typename OpType>
@@ -122,7 +130,7 @@ template<typename OpType>
  * @return the result of the async op
  */
 template<typename OpType>
-[[nodiscard]] UE::Online::TOnlineResult<OpType> BlockUntilAsyncOpComplete(UE::Online::TOnlineAsyncOpHandle<OpType> InAsyncOp, TOptional<int32> MaxTicks = BlockUntilCompleteDefaultMaxTicks, ESleepBehavior SleepBehavior = ESleepBehavior::NoSleep)
+UE::Online::TOnlineResult<OpType> GetAsyncOpResultBlocking(UE::Online::TOnlineAsyncOpHandle<OpType> InAsyncOp, TOptional<int32> MaxTicks = BlockUntilCompleteDefaultMaxTicks, ESleepBehavior SleepBehavior = ESleepBehavior::NoSleep)
 {
 	using namespace UE::Online;
 	TSharedRef<TPromise<TOnlineResult<OpType>>> Promise = MakeShared<TPromise<TOnlineResult<OpType>>>();
@@ -131,21 +139,7 @@ template<typename OpType>
 	{
 		Promise->EmplaceValue(Result);
 	});
-	BlockUntilComplete(Future, MaxTicks, SleepBehavior);
-	return Future.Get();
-}
-
-/**
- * Block until an async op completes, and return the result as an FOnlineError.
- * Will return Errors::Success on success.
- * Will fail if the async op does not complete in a timely manner.
- * @param InAsyncOp the async op to wait for
- * @return the FOnlineError result of the async op. Will be Errors::Success if the op was successful.
- */
-template<typename OpType>
-[[nodiscard]] UE::Online::FOnlineError GetAsyncOpResultBasic(UE::Online::TOnlineAsyncOpHandle<OpType> InAsyncOp)
-{
-	return GetErrorFromResult(BlockUntilAsyncOpComplete(InAsyncOp));
+	return GetFutureValueBlocking(Future, MaxTicks, SleepBehavior);
 }
 
 /**
@@ -157,7 +151,7 @@ template<typename OpType>
 template<typename OpType>
 [[nodiscard]] OpType::Result GetAsyncOpResultChecked(UE::Online::TOnlineAsyncOpHandle<OpType> InAsyncOp, TOptional<int32> MaxTicks = BlockUntilCompleteDefaultMaxTicks, ESleepBehavior SleepBehavior = ESleepBehavior::NoSleep)
 {
-	return GetOpResultChecked(BlockUntilAsyncOpComplete(InAsyncOp, MaxTicks, SleepBehavior));
+	return GetOpResultChecked(GetAsyncOpResultBlocking(InAsyncOp, MaxTicks, SleepBehavior));
 }
 
 /* UE::TestCommon */ }
\ No newline at end of file
