- Fix unused params

diff --git a/Engine/Source/Programs/Online/OnlineTestsCore/Source/Public/Helpers/OnlineHelpers.h b/Engine/Source/Programs/Online/OnlineTestsCore/Source/Public/Helpers/OnlineHelpers.h
--- a/Engine/Source/Programs/Online/OnlineTestsCore/Source/Public/Helpers/OnlineHelpers.h
+++ b/Engine/Source/Programs/Online/OnlineTestsCore/Source/Public/Helpers/OnlineHelpers.h
@@ -62,14 +62,14 @@ void BlockUntilComplete(TArrayView<const TFuture<ResultType>> Futures, TOptional
 template<typename ResultType>
 void BlockUntilComplete(const TArray<TFuture<ResultType>>& Futures, TOptional<int32> MaxTicks = BlockUntilCompleteDefaultMaxTicks, ESleepBehavior SleepBehavior = ESleepBehavior::NoSleep)
 {
-	BlockUntilComplete(MakeArrayView(const_cast<const TFuture<ResultType>*>(Futures.GetData()), Futures.Num()), BlockUntilCompleteDefaultMaxTicks, SleepBehavior);
+	BlockUntilComplete(MakeArrayView(const_cast<const TFuture<ResultType>*>(Futures.GetData()), Futures.Num()), MaxTicks, SleepBehavior);
 }
 
 // Overload that takes a single future
 template<typename ResultType>
 void BlockUntilComplete(const TFuture<ResultType>& Future, TOptional<int32> MaxTicks = BlockUntilCompleteDefaultMaxTicks, ESleepBehavior SleepBehavior = ESleepBehavior::NoSleep)
 {
-	BlockUntilComplete(MakeArrayView(&Future, 1), BlockUntilCompleteDefaultMaxTicks, SleepBehavior);
+	BlockUntilComplete(MakeArrayView(&Future, 1), MaxTicks, SleepBehavior);
 }
 
 /** Block until a future is complete, and return its value */
