Suppressing warnings in OsMcpTests.

diff --git a/Engine/Source/Programs/Online/OSSTestsCore/Tests/Purchase/PurchaseGetReceiptsTests.cpp b/Engine/Source/Programs/Online/OSSTestsCore/Tests/Purchase/PurchaseGetReceiptsTests.cpp
--- a/Engine/Source/Programs/Online/OSSTestsCore/Tests/Purchase/PurchaseGetReceiptsTests.cpp
+++ b/Engine/Source/Programs/Online/OSSTestsCore/Tests/Purchase/PurchaseGetReceiptsTests.cpp
@@ -27,7 +27,6 @@ PURCHASE_TEST_CASE("Verify calling GetReceipts with valid inputs returns the exp
 	bool bRestoreReceipts = false;
 	bool bIsShouldBeEmpty = false;
 
-	int32 NumUsersToImplicitLogin = 1;
 	TArray<FPurchaseReceipt>* PurchaseReceipts = nullptr;
 
 	FPurchaseCheckoutRequest PurchCheckRequest = {};
