Fix non-unity errors
#rnx
#Synced-CL 50445837

diff --git a/Engine/Source/Developer/ZenOplogUtils/Tests/ZenOplogUtilsTest.cpp b/Engine/Source/Developer/ZenOplogUtils/Tests/ZenOplogUtilsTest.cpp
--- a/Engine/Source/Developer/ZenOplogUtils/Tests/ZenOplogUtilsTest.cpp
+++ b/Engine/Source/Developer/ZenOplogUtils/Tests/ZenOplogUtilsTest.cpp
@@ -9,6 +9,7 @@
 #include "Algo/IsSorted.h"
 #include "Algo/Transform.h"
 #include "Serialization/CompactBinarySerialization.h"
+#include "Serialization/CompactBinaryWriter.h"
 
 #if WITH_EDITORONLY_DATA
 
diff --git a/Engine/Source/Programs/Online/OnlineTestsCore/Source/Private/TestUtilities.h b/Engine/Source/Programs/Online/OnlineTestsCore/Source/Private/TestUtilities.h
--- a/Engine/Source/Programs/Online/OnlineTestsCore/Source/Private/TestUtilities.h
+++ b/Engine/Source/Programs/Online/OnlineTestsCore/Source/Private/TestUtilities.h
@@ -2,6 +2,8 @@
 #pragma once
 
 #include "TestHarness.h"
+#include "Containers/UnrealString.h"
+#include "Misc/Guid.h"
 
 class FTestUtilities
 {
diff --git a/Engine/Source/Programs/Online/OnlineTestsCore/Source/Public/Helpers/Achievements/QueryAchievementStatesHelper.h b/Engine/Source/Programs/Online/OnlineTestsCore/Source/Public/Helpers/Achievements/QueryAchievementStatesHelper.h
--- a/Engine/Source/Programs/Online/OnlineTestsCore/Source/Public/Helpers/Achievements/QueryAchievementStatesHelper.h
+++ b/Engine/Source/Programs/Online/OnlineTestsCore/Source/Public/Helpers/Achievements/QueryAchievementStatesHelper.h
@@ -3,6 +3,7 @@
 #pragma once
 #include "TestDriver.h"
 #include "TestHarness.h"
+#include "OnlineCatchHelper.h"
 #include "OnlineCatchStringMakers.h"
 
 #include "Online/AuthCommon.h"
diff --git a/Engine/Source/Programs/Online/OnlineTestsCore/Source/Public/Helpers/Auth/AuthQueryExternalAuthToken.h b/Engine/Source/Programs/Online/OnlineTestsCore/Source/Public/Helpers/Auth/AuthQueryExternalAuthToken.h
--- a/Engine/Source/Programs/Online/OnlineTestsCore/Source/Public/Helpers/Auth/AuthQueryExternalAuthToken.h
+++ b/Engine/Source/Programs/Online/OnlineTestsCore/Source/Public/Helpers/Auth/AuthQueryExternalAuthToken.h
@@ -5,6 +5,7 @@
 #include "TestDriver.h"
 #include "TestHarness.h"
 
+#include "AsyncTestStep.h"
 #include "Online/AuthCommon.h"
 #include "Online/Auth.h"
 #include "Online/OnlineAsyncOp.h"
diff --git a/Engine/Source/Programs/Online/OnlineTestsCore/Source/Public/Helpers/Sessions/JoinSessionHelper.h b/Engine/Source/Programs/Online/OnlineTestsCore/Source/Public/Helpers/Sessions/JoinSessionHelper.h
--- a/Engine/Source/Programs/Online/OnlineTestsCore/Source/Public/Helpers/Sessions/JoinSessionHelper.h
+++ b/Engine/Source/Programs/Online/OnlineTestsCore/Source/Public/Helpers/Sessions/JoinSessionHelper.h
@@ -4,6 +4,7 @@
 
 #include "Online/SessionsCommon.h"
 #include "AsyncTestStep.h"
+#include "OnlineCatchHelper.h"
 
 struct FJoinSessionHelper : public FAsyncTestStep
 {
diff --git a/Engine/Source/Programs/Online/OnlineTestsCore/Source/Public/Helpers/UserFile/AsyncUserFileDeleteFile.h b/Engine/Source/Programs/Online/OnlineTestsCore/Source/Public/Helpers/UserFile/AsyncUserFileDeleteFile.h
--- a/Engine/Source/Programs/Online/OnlineTestsCore/Source/Public/Helpers/UserFile/AsyncUserFileDeleteFile.h
+++ b/Engine/Source/Programs/Online/OnlineTestsCore/Source/Public/Helpers/UserFile/AsyncUserFileDeleteFile.h
@@ -60,7 +60,7 @@ public:
 					CHECK(Result.IsOk() == ExpectedResult.IsOk());
 					if (Result.IsError())
 					{
-						UE::Online::FOnlineError Error = ExpectedResult.IsError() ? ExpectedResult.GetErrorValue() : UE::Online::FOnlineError(Errors::Unknown());
+						UE::Online::FOnlineError Error = ExpectedResult.IsError() ? ExpectedResult.GetErrorValue() : UE::Online::FOnlineError(UE::Online::Errors::Unknown());
 						CHECK(Result.GetErrorValue() == Error);
 					}
 				}
