Fix various unreachable code, duplicate symbol, include case, and missing symbol errors
#rnx

#Synced-CL 50505807

diff --git a/Engine/Plugins/Animation/RigLogic/Source/RigLogicModule/Public/RigUnit_RigLogic_Test.h b/Engine/Plugins/Animation/RigLogic/Source/RigLogicModule/Public/RigUnit_RigLogic_Test.h
--- a/Engine/Plugins/Animation/RigLogic/Source/RigLogicModule/Public/RigUnit_RigLogic_Test.h
+++ b/Engine/Plugins/Animation/RigLogic/Source/RigLogicModule/Public/RigUnit_RigLogic_Test.h
@@ -2,6 +2,8 @@
 
 #pragma once
 
+#if WITH_DEV_AUTOMATION_TESTS
+
 #include "CoreTypes.h"
 #include "CoreMinimal.h"
 #include "UObject/StrongObjectPtr.h"
@@ -9,9 +11,7 @@
 
 #include "DNAReader.h"
 
-#if WITH_DEV_AUTOMATION_TESTS
-	#include "RigUnit_RigLogic.h"
-#endif
+#include "RigUnit_RigLogic.h"
 
 struct FRigUnit_RigLogic_JointUpdateParams;
 
@@ -1026,4 +1026,6 @@ private:
 
 public:
 	FRigUnit_RigLogic_Data* GetData() { return &(Unit->Data); };
-};
\ No newline at end of file
+};
+
+#endif
\ No newline at end of file
diff --git a/Engine/Source/Programs/Online/OnlineTestsCore/Source/Public/Helpers/Presence/CommonPresenceHelpers.h b/Engine/Source/Programs/Online/OnlineTestsCore/Source/Public/Helpers/Presence/CommonPresenceHelpers.h
--- a/Engine/Source/Programs/Online/OnlineTestsCore/Source/Public/Helpers/Presence/CommonPresenceHelpers.h
+++ b/Engine/Source/Programs/Online/OnlineTestsCore/Source/Public/Helpers/Presence/CommonPresenceHelpers.h
@@ -5,7 +5,7 @@
 #include "Online/Presence.h"
 using namespace UE::Online;
 
-TSharedRef<FUserPresence> GetGenericPresenceA(FAccountId AccountId)
+inline TSharedRef<FUserPresence> GetGenericPresenceA(FAccountId AccountId)
 {
 	TSharedRef<FUserPresence> Presence = MakeShared<FUserPresence>();
 
@@ -20,7 +20,7 @@ TSharedRef<FUserPresence> GetGenericPresenceA(FAccountId AccountId)
 	return Presence;
 }
 
-TSharedRef<FUserPresence> GetGenericPresenceB(FAccountId AccountId)
+inline TSharedRef<FUserPresence> GetGenericPresenceB(FAccountId AccountId)
 {
 	TSharedRef<FUserPresence> Presence = MakeShared<FUserPresence>();
 
@@ -35,7 +35,7 @@ TSharedRef<FUserPresence> GetGenericPresenceB(FAccountId AccountId)
 	return Presence;
 }
 
-FPartialUpdatePresence::Params GetGenericPresenceMutationAtoB(FAccountId AccountId)
+inline FPartialUpdatePresence::Params GetGenericPresenceMutationAtoB(FAccountId AccountId)
 {
 	FPartialUpdatePresence::Params Params;
 
@@ -51,7 +51,7 @@ FPartialUpdatePresence::Params GetGenericPresenceMutationAtoB(FAccountId Account
 	return Params;
 }
 
-void CheckPresenceAreEqual(const FUserPresence& Presence1, const FUserPresence& Presence2)
+inline void CheckPresenceAreEqual(const FUserPresence& Presence1, const FUserPresence& Presence2)
 {
 	CHECK(Presence1.AccountId == Presence2.AccountId);
 	CHECK(Presence1.GameStatus == Presence2.GameStatus);
@@ -73,7 +73,7 @@ void CheckPresenceAreEqual(const FUserPresence& Presence1, const FUserPresence&
 }
 
 // don't check user ids here- may be used to check for an update on a specific user
-void CheckPresencesNotEqual(const FUserPresence& Presence1, const FUserPresence& Presence2)
+inline void CheckPresencesNotEqual(const FUserPresence& Presence1, const FUserPresence& Presence2)
 {
 	CHECK(Presence1.GameStatus != Presence2.GameStatus);
 	CHECK(Presence1.Joinability != Presence2.Joinability);
diff --git a/Engine/Source/Runtime/ApplicationCore/Private/Linux/LinuxPlatformApplicationMisc.cpp b/Engine/Source/Runtime/ApplicationCore/Private/Linux/LinuxPlatformApplicationMisc.cpp
--- a/Engine/Source/Runtime/ApplicationCore/Private/Linux/LinuxPlatformApplicationMisc.cpp
+++ b/Engine/Source/Runtime/ApplicationCore/Private/Linux/LinuxPlatformApplicationMisc.cpp
@@ -215,7 +215,7 @@ EAppReturnType::Type MessageBoxExtImpl(EAppMsgType::Type MsgType, const TCHAR* T
 	EAppReturnType::Type Answer = EAppReturnType::Type::Cancel;
 
 	FSlowHeartBeatScope SuspendHeartBeat;
-	if (SDL_ShowMessageBox(&MessageBoxData, &ButtonPressed) == -1)
+	if (!SDL_ShowMessageBox(&MessageBoxData, &ButtonPressed))
 	{
 		FString Message = GetHeadlessMessageBoxMessage(MsgType, Caption, Text, Answer);
 		UE_LOG(LogLinux, Warning, TEXT("%s"), *Message);
diff --git a/Engine/Source/Runtime/Core/Private/Android/AndroidPlatformStackWalk.cpp b/Engine/Source/Runtime/Core/Private/Android/AndroidPlatformStackWalk.cpp
--- a/Engine/Source/Runtime/Core/Private/Android/AndroidPlatformStackWalk.cpp
+++ b/Engine/Source/Runtime/Core/Private/Android/AndroidPlatformStackWalk.cpp
@@ -226,10 +226,19 @@ uint32 FAndroidPlatformStackWalk::CaptureStackBackTrace(uint64* BackTrace, uint3
 	}
 
 	// zero results
-	FPlatformMemory::Memzero(BackTrace, MaxDepth*sizeof(uint64));
+	FPlatformMemory::Memzero(BackTrace, MaxDepth * sizeof(uint64));
 	
+	auto DefaultUnwindBacktrace = [BackTrace, MaxDepth]()
+	{
+		AndroidStackWalkHelpers::BackTrace = BackTrace;
+		AndroidStackWalkHelpers::MaxDepth = MaxDepth;
+		uint32 Depth = 0;
+		_Unwind_Backtrace(AndroidStackWalkHelpers::BacktraceCallback, &Depth);
+		return Depth;
+	};
+
 #if PLATFORM_ANDROID_ARM
-	if (Context != nullptr)
+	if (Context)
 	{
 		// Android signal handlers always catch signals before user handlers and passes it down to user later
 		// _Unwind_Backtrace does not use signal context and will produce wrong callstack in this case
@@ -237,6 +246,10 @@ uint32 FAndroidPlatformStackWalk::CaptureStackBackTrace(uint64* BackTrace, uint3
 		// Code taken from https://android.googlesource.com/platform/system/core/+/jb-dev/libcorkscrew/arch-arm/backtrace-arm.c
 		return unwind_backtrace_signal(Context, BackTrace, MaxDepth);
 	}
+	else
+	{
+		return DefaultUnwindBacktrace();
+	}
 #elif HAS_LIBUNWIND
 	if (Context)
 	{
@@ -249,13 +262,10 @@ uint32 FAndroidPlatformStackWalk::CaptureStackBackTrace(uint64* BackTrace, uint3
 	{
 		return unw_backtrace((void**)BackTrace, MaxDepth);
 	}
+#else
+	return DefaultUnwindBacktrace();
 #endif 
 
-	AndroidStackWalkHelpers::BackTrace = BackTrace;
-	AndroidStackWalkHelpers::MaxDepth = MaxDepth;
-	uint32 Depth = 0;
-	_Unwind_Backtrace(AndroidStackWalkHelpers::BacktraceCallback, &Depth);
-	return Depth;
 }
 
 uint32 FAndroidPlatformStackWalk::CaptureStackBackTraceViaFramePointerWalking(uint64* BackTrace, uint32 MaxDepth)
diff --git a/Engine/Source/Runtime/CoreUObject/Private/UObject/GarbageCollection.cpp b/Engine/Source/Runtime/CoreUObject/Private/UObject/GarbageCollection.cpp
--- a/Engine/Source/Runtime/CoreUObject/Private/UObject/GarbageCollection.cpp
+++ b/Engine/Source/Runtime/CoreUObject/Private/UObject/GarbageCollection.cpp
@@ -6044,7 +6044,10 @@ void FReachabilityAnalysisState::PerformReachabilityAnalysisAndConditionallyPurg
 				const double ReferenceProcessingTotalTimeMs = ReferenceProcessingTotalTime * 1000;
 				const double IncrementalMarkPhaseTotalTimeMs = IncrementalMarkPhaseTotalTime * 1000;
 				UE_LOG(LogGarbage, Log, TEXT("GC Reachability Analysis total time: %.2f ms (%.2f ms on reference traversal)"), IncrementalMarkPhaseTotalTimeMs, ReferenceProcessingTotalTimeMs);
-				FString ExtraDetail = WITH_VERSE_VM ? FString::Printf(TEXT("and %d Verse cells"), Stats.NumVerseCells) : FString();
+				FString ExtraDetail;
+#if WITH_VERSE_VM
+				ExtraDetail = FString::Printf(TEXT("and %d Verse cells"), Stats.NumVerseCells);
+#endif
 				UE_LOG(LogGarbage, Log, TEXT("%.2f ms for %sGC - %d refs/ms while processing %d references from %d objects %s with %d clusters"),
 					IncrementalMarkPhaseTotalTimeMs,
 					bForceNonIncrementalReachability ? TEXT("") : TEXT("Incremental "),
diff --git a/Engine/Source/Runtime/CoreUObject/Private/UObject/ReferenceChainSearch.cpp b/Engine/Source/Runtime/CoreUObject/Private/UObject/ReferenceChainSearch.cpp
--- a/Engine/Source/Runtime/CoreUObject/Private/UObject/ReferenceChainSearch.cpp
+++ b/Engine/Source/Runtime/CoreUObject/Private/UObject/ReferenceChainSearch.cpp
@@ -78,6 +78,7 @@ namespace UE::ReferenceChainSearch
 			return VertexToObject(PotentialInner)->IsIn(VertexToObject(Outer));
 		}
 
+		PRAGMA_DISABLE_UNREACHABLE_CODE_WARNINGS
 		bool IsRoot(FVertex Vertex, EReferenceChainSearchMode SearchMode) const
 		{
 			const UObject* Object = VertexToObject(Vertex);
@@ -86,6 +87,7 @@ namespace UE::ReferenceChainSearch
 				|| (GARBAGE_COLLECTION_KEEPFLAGS != RF_NoFlags && Object->HasAnyFlags(GARBAGE_COLLECTION_KEEPFLAGS)
 					&& !(SearchMode & EReferenceChainSearchMode::FullChain));
 		}
+		PRAGMA_RESTORE_UNREACHABLE_CODE_WARNINGS
 
 		bool IsGarbage(FVertex Vertex) const
 		{
@@ -189,6 +191,7 @@ namespace UE::ReferenceChainSearch
 			return VertexToObject(PotentialInner)->IsIn(VertexToObject(Outer));
 		}
 
+		PRAGMA_DISABLE_UNREACHABLE_CODE_WARNINGS
 		bool IsRoot(FVertex Vertex, EReferenceChainSearchMode SearchMode) const
 		{
 			ObjectType Object = VertexToObject(Vertex);
@@ -197,6 +200,7 @@ namespace UE::ReferenceChainSearch
 				|| (GARBAGE_COLLECTION_KEEPFLAGS != RF_NoFlags && Object->HasAnyFlags(GARBAGE_COLLECTION_KEEPFLAGS)
 					&& !(SearchMode & EReferenceChainSearchMode::FullChain));
 		}
+		PRAGMA_RESTORE_UNREACHABLE_CODE_WARNINGS
 
 		bool IsGarbage(FVertex Vertex) const
 		{
diff --git a/Engine/Source/Runtime/Experimental/ChaosSolverEngine/Public/Chaos/ChaosVDTraceRelayTransport.h b/Engine/Source/Runtime/Experimental/ChaosSolverEngine/Public/Chaos/ChaosVDTraceRelayTransport.h
--- a/Engine/Source/Runtime/Experimental/ChaosSolverEngine/Public/Chaos/ChaosVDTraceRelayTransport.h
+++ b/Engine/Source/Runtime/Experimental/ChaosSolverEngine/Public/Chaos/ChaosVDTraceRelayTransport.h
@@ -4,7 +4,7 @@
 
 #include "Containers/Array.h"
 #include "Containers/UnrealString.h"
-#include "Misc/CoreMiscdefines.h"
+#include "Misc/CoreMiscDefines.h"
 
 namespace Chaos::VD
 {
diff --git a/Engine/Source/Runtime/Experimental/IoStore/OnDemand/Private/Statistics.h b/Engine/Source/Runtime/Experimental/IoStore/OnDemand/Private/Statistics.h
--- a/Engine/Source/Runtime/Experimental/IoStore/OnDemand/Private/Statistics.h
+++ b/Engine/Source/Runtime/Experimental/IoStore/OnDemand/Private/Statistics.h
@@ -5,6 +5,7 @@
 #include "CVarUtilities.h"
 #include "HAL/LowLevelMemTracker.h"
 #include "IO/OnDemandError.h"
+#include "IO/IoStoreOnDemand.h"
 #include "Misc/CoreDelegates.h"
 #include "OnDemandBackendStatus.h"
 #include "ProfilingDebugging/CountersTrace.h"
diff --git a/Engine/Source/Runtime/WebBrowser/Private/WebBrowserSingleton.cpp b/Engine/Source/Runtime/WebBrowser/Private/WebBrowserSingleton.cpp
--- a/Engine/Source/Runtime/WebBrowser/Private/WebBrowserSingleton.cpp
+++ b/Engine/Source/Runtime/WebBrowser/Private/WebBrowserSingleton.cpp
@@ -422,8 +422,10 @@ FWebBrowserSingleton::FWebBrowserSingleton(const FWebBrowserInitSettings& WebBro
 #endif
 
 		// Initialize CEF.
+		PRAGMA_DISABLE_UNREACHABLE_CODE_WARNINGS // only windows uses retries, but restructuring for other platforms is not worth it
 		for (int NumRetries = 0; NumRetries <= CEF3_INIT_MAX_RETRIES; ++NumRetries)
 		{
+		PRAGMA_RESTORE_UNREACHABLE_CODE_WARNINGS
 			bCEFInitialized = CefInitialize(MainArgs, Settings, CEFBrowserApp.get(), nullptr);
 			if (bCEFInitialized)
 			{
