- Fix spelling of Hierarchy

#Synced-CL 50524808

#Synced-CL 50524808

diff --git a/Engine/Plugins/Online/OnlineServices/Source/OnlineServicesCommon/Private/Online/OnlineServicesCommon.cpp b/Engine/Plugins/Online/OnlineServices/Source/OnlineServicesCommon/Private/Online/OnlineServicesCommon.cpp
--- a/Engine/Plugins/Online/OnlineServices/Source/OnlineServicesCommon/Private/Online/OnlineServicesCommon.cpp
+++ b/Engine/Plugins/Online/OnlineServices/Source/OnlineServicesCommon/Private/Online/OnlineServicesCommon.cpp
@@ -46,9 +46,9 @@ FOnlineServicesCommon::~FOnlineServicesCommon()
 void FOnlineServicesCommon::Init()
 {
 	OpCache.SetLoadConfigFn(
-		[this](FOperationConfig& OperationConfig, const TArray<FString>& SectionHeirarchy)
+		[this](FOperationConfig& OperationConfig, const TArray<FString>& SectionHierarchy)
 		{
-			return LoadConfig(OperationConfig, SectionHeirarchy);
+			return LoadConfig(OperationConfig, SectionHierarchy);
 		});
 
 	RegisterComponents();
diff --git a/Engine/Plugins/Online/OnlineServices/Source/OnlineServicesCommon/Public/Online/OnlineAsyncOpCache.h b/Engine/Plugins/Online/OnlineServices/Source/OnlineServicesCommon/Public/Online/OnlineAsyncOpCache.h
--- a/Engine/Plugins/Online/OnlineServices/Source/OnlineServicesCommon/Public/Online/OnlineAsyncOpCache.h
+++ b/Engine/Plugins/Online/OnlineServices/Source/OnlineServicesCommon/Public/Online/OnlineAsyncOpCache.h
@@ -246,7 +246,7 @@ public:
 
 	// Create an operation
 	template <typename OpType>
-	TOnlineAsyncOpRef<OpType> GetOp(typename OpType::Params&& Params, const TArray<FString>& ConfigSectionHeirarchy)
+	TOnlineAsyncOpRef<OpType> GetOp(typename OpType::Params&& Params, const TArray<FString>& ConfigSectionHierarchy)
 	{
 		TUniquePtr<TWrappedOperation<OpType>> WrappedOp = MakeUnique<TWrappedOperation<OpType>>(Services, MoveTemp(Params));
 		TSharedRef<TOnlineAsyncOp<OpType>> Op = WrappedOp->GetDataRef();
@@ -276,7 +276,7 @@ public:
 
 	// Join an existing operation or use a non-expired cached result, or create an operation that can later be joined
 	template <typename OpType, typename ParamsFuncsType /*= TJoinableOpParamsFuncs<OpType>*/>
-	TOnlineAsyncOpRef<OpType> GetJoinableOp(typename OpType::Params&& Params, const TArray<FString>& ConfigSectionHeirarchy)
+	TOnlineAsyncOpRef<OpType> GetJoinableOp(typename OpType::Params&& Params, const TArray<FString>& ConfigSectionHierarchy)
 	{
 		TSharedPtr<TOnlineAsyncOp<OpType>> Op;
 
@@ -306,7 +306,7 @@ public:
 		{
 			FOperationConfig Config;
 			check(LoadConfigFn);
-			LoadConfigFn(Config, ConfigSectionHeirarchy);
+			LoadConfigFn(Config, ConfigSectionHierarchy);
 			
 			Op = CreateOp<OpType, ParamsFuncsType>(MoveTemp(Params));
 			Op->OpCacheHandle = Op->OnComplete().Add(GetSharedThis(), [this, Config](const TOnlineAsyncOp<OpType>& ThisOp, const TOnlineResult<OpType>& Result)
@@ -352,7 +352,7 @@ public:
 
 	// Merge with a pending operation, or create an operation
 	template <typename OpType, typename ParamsFuncsType /*= TMergeableOpParamsFuncs<OpType>*/>
-	TOnlineAsyncOpRef<OpType> GetMergeableOp(typename OpType::Params&& Params, const TArray<FString>& ConfigSectionHeirarchy)
+	TOnlineAsyncOpRef<OpType> GetMergeableOp(typename OpType::Params&& Params, const TArray<FString>& ConfigSectionHierarchy)
 	{
 		TSharedPtr<TOnlineAsyncOp<OpType>> Op;
 
@@ -380,7 +380,7 @@ public:
 		{
 			FOperationConfig Config;
 			check(LoadConfigFn);
-			LoadConfigFn(Config, ConfigSectionHeirarchy);
+			LoadConfigFn(Config, ConfigSectionHierarchy);
 
 			Op = CreateOp<OpType, ParamsFuncsType>(MoveTemp(Params));
 			// remove from cache once operation has started. It is no longer mergeable at that point
diff --git a/Engine/Plugins/Online/OnlineServices/Source/OnlineServicesCommon/Public/Online/OnlineComponent.h b/Engine/Plugins/Online/OnlineServices/Source/OnlineServicesCommon/Public/Online/OnlineComponent.h
--- a/Engine/Plugins/Online/OnlineServices/Source/OnlineServicesCommon/Public/Online/OnlineComponent.h
+++ b/Engine/Plugins/Online/OnlineServices/Source/OnlineServicesCommon/Public/Online/OnlineComponent.h
@@ -52,7 +52,7 @@ public:
 	template <typename OpType>
 	TOnlineAsyncOpRef<OpType> GetOp(typename OpType::Params&& Params)
 	{
-		TOnlineAsyncOpRef<OpType> Op = Services.template GetOp<OpType>(MoveTemp(Params), GetConfigSectionHeirarchy(OpType::Name));
+		TOnlineAsyncOpRef<OpType> Op = Services.template GetOp<OpType>(MoveTemp(Params), GetConfigSectionHierarchy(OpType::Name));
 		Op->SetInterfaceName(InterfaceName);
 		return Op;
 	}
@@ -60,7 +60,7 @@ public:
 	template <typename OpType, typename ParamsFuncsType = TJoinableOpParamsFuncs<OpType>>
 	TOnlineAsyncOpRef<OpType> GetJoinableOp(typename OpType::Params&& Params)
 	{
-		TOnlineAsyncOpRef<OpType> Op = Services.template GetJoinableOp<OpType, ParamsFuncsType>(MoveTemp(Params), GetConfigSectionHeirarchy(OpType::Name));
+		TOnlineAsyncOpRef<OpType> Op = Services.template GetJoinableOp<OpType, ParamsFuncsType>(MoveTemp(Params), GetConfigSectionHierarchy(OpType::Name));
 		Op->SetInterfaceName(InterfaceName);
 		return Op;
 	}
@@ -68,7 +68,7 @@ public:
 	template <typename OpType, typename ParamsFuncsType = TMergeableOpParamsFuncs<OpType>>
 	TOnlineAsyncOpRef<OpType> GetMergeableOp(typename OpType::Params&& Params)
 	{
-		TOnlineAsyncOpRef<OpType> Op = Services.template GetMergeableOp<OpType, ParamsFuncsType>(MoveTemp(Params), GetConfigSectionHeirarchy(OpType::Name));
+		TOnlineAsyncOpRef<OpType> Op = Services.template GetMergeableOp<OpType, ParamsFuncsType>(MoveTemp(Params), GetConfigSectionHierarchy(OpType::Name));
 		Op->SetInterfaceName(InterfaceName);
 		return Op;
 	}
@@ -103,27 +103,33 @@ public:
 		return **Queue;
 	}
 
-	UE_DEPRECATED(5.6, "GetConfigSectionHeiarchy has been renamed GetConfigSectionHeirarchy")
+	UE_DEPRECATED(5.6, "GetConfigSectionHeiarchy has been renamed GetConfigSectionHierarchy")
 	TArray<FString> GetConfigSectionHeiarchy(const FString& OperationName = FString())
 	{
-		return GetConfigSectionHeirarchy(OperationName);
+		return GetConfigSectionHierarchy(OperationName);
 	}
 
+	UE_DEPRECATED(5.8, "GetConfigSectionHeirarchy has been renamed GetConfigSectionHierarchy")
 	TArray<FString> GetConfigSectionHeirarchy(const FString& OperationName = FString())
 	{
-		TArray<FString> SectionHeirarchy;
+		return GetConfigSectionHierarchy(OperationName);
+	}
+
+	TArray<FString> GetConfigSectionHierarchy(const FString& OperationName = FString())
+	{
+		TArray<FString> SectionHierarchy;
 		FString SectionName = TEXT("OnlineServices");
-		SectionHeirarchy.Add(SectionName);
+		SectionHierarchy.Add(SectionName);
 		SectionName += TEXT(".") + GetServices().GetServiceConfigName();
-		SectionHeirarchy.Add(SectionName);
+		SectionHierarchy.Add(SectionName);
 		SectionName += TEXT(".") + GetConfigName();
-		SectionHeirarchy.Add(SectionName);
+		SectionHierarchy.Add(SectionName);
 		if (!OperationName.IsEmpty())
 		{
 			SectionName += TEXT(".") + OperationName;
-			SectionHeirarchy.Add(SectionName);
+			SectionHierarchy.Add(SectionName);
 		}
-		return SectionHeirarchy;
+		return SectionHierarchy;
 	}
 
 	virtual void RegisterCommands() {}
diff --git a/Engine/Plugins/Online/OnlineServices/Source/OnlineServicesCommon/Public/Online/OnlineConfig.h b/Engine/Plugins/Online/OnlineServices/Source/OnlineServicesCommon/Public/Online/OnlineConfig.h
--- a/Engine/Plugins/Online/OnlineServices/Source/OnlineServicesCommon/Public/Online/OnlineConfig.h
+++ b/Engine/Plugins/Online/OnlineServices/Source/OnlineServicesCommon/Public/Online/OnlineConfig.h
@@ -384,17 +384,17 @@ auto LoadConfig(IOnlineConfigProvider& Provider, const FString& Section, T& Valu
  * Requires that the type metadata is specified for StructType
  *
  * @param Provider The config provider
- * @param SectionHeirarchy The config sections to load values from
+ * @param SectionHierarchy The config sections to load values from
  * @param Value The struct whose fields will be loaded from config
  *
  * @return true if any values were loaded
  */
 template <typename T>
-auto LoadConfig(IOnlineConfigProvider& Provider, const TArray<FString>& SectionHeirarchy, T& Value)
+auto LoadConfig(IOnlineConfigProvider& Provider, const TArray<FString>& SectionHierarchy, T& Value)
 	-> std::enable_if_t<TModels_V<Meta::COnlineMetadataAvailable, T>, bool>
 {
 	bool bLoadedValue = false;
-	for (const FString& Section : SectionHeirarchy)
+	for (const FString& Section : SectionHierarchy)
 	{
 		bool bLoadedInCurrentSection = LoadConfig(Provider, Section, Value);
 		if (bLoadedInCurrentSection && Section.Contains(" "))
@@ -444,18 +444,18 @@ auto LoadConfig(IOnlineConfigProvider& Provider, const FString& Section, const T
  * Requires that the type metadata is specified for StructType, and that the config value is in config struct syntax.
  *
  * @param Provider The config provider
- * @param SectionHeirarchy The config sections to load values from
+ * @param SectionHierarchy The config sections to load values from
  * @param Key The key in the config section where the struct resides
  * @param Value The struct whose fields will be loaded from config
  *
  * @return true if any values were loaded
  */
 template <typename T>
-auto LoadConfig(IOnlineConfigProvider& Provider, const TArray<FString>& SectionHeirarchy, const TCHAR* Key, T& Value)
+auto LoadConfig(IOnlineConfigProvider& Provider, const TArray<FString>& SectionHierarchy, const TCHAR* Key, T& Value)
 -> std::enable_if_t<TModels_V<Meta::COnlineMetadataAvailable, T>, bool>
 {
 	bool bLoadedValue = false;
-	for (const FString& Section : SectionHeirarchy)
+	for (const FString& Section : SectionHierarchy)
 	{
 		bLoadedValue |= LoadConfig(Provider, Section, Key, Value);
 	}
diff --git a/Engine/Plugins/Online/OnlineServices/Source/OnlineServicesCommon/Public/Online/OnlineServicesCommon.h b/Engine/Plugins/Online/OnlineServices/Source/OnlineServicesCommon/Public/Online/OnlineServicesCommon.h
--- a/Engine/Plugins/Online/OnlineServices/Source/OnlineServicesCommon/Public/Online/OnlineServicesCommon.h
+++ b/Engine/Plugins/Online/OnlineServices/Source/OnlineServicesCommon/Public/Online/OnlineServicesCommon.h
@@ -173,25 +173,31 @@ public:
 	 */
 	const FString& GetServiceConfigName() const { return ServiceConfigName; }
 
-	UE_DEPRECATED(5.6, "GetConfigSectionHeiarchy has been renamed GetConfigSectionHeirarchy")
+	UE_DEPRECATED(5.6, "GetConfigSectionHeiarchy has been renamed GetConfigSectionHierarchy")
 	TArray<FString> GetConfigSectionHeiarchy(const FString& OperationName = FString()) const
 	{
-		return GetConfigSectionHeirarchy(OperationName);
+		return GetConfigSectionHierarchy(OperationName);
 	}
-	
+
+	UE_DEPRECATED(5.8, "GetConfigSectionHeirarchy has been renamed GetConfigSectionHierarchy")
 	TArray<FString> GetConfigSectionHeirarchy(const FString& OperationName = FString()) const
 	{
-		TArray<FString> SectionHeirarchy;
+		return GetConfigSectionHierarchy(OperationName);
+	}
+	
+	TArray<FString> GetConfigSectionHierarchy(const FString& OperationName = FString()) const
+	{
+		TArray<FString> SectionHierarchy;
 		FString SectionName = TEXT("OnlineServices");
-		SectionHeirarchy.Add(SectionName);
+		SectionHierarchy.Add(SectionName);
 		SectionName += TEXT(".") + GetServiceConfigName();
-		SectionHeirarchy.Add(SectionName);
+		SectionHierarchy.Add(SectionName);
 		if (!OperationName.IsEmpty())
 		{
 			SectionName += TEXT(".") + OperationName;
-			SectionHeirarchy.Add(SectionName);
+			SectionHierarchy.Add(SectionName);
 		}
-		return SectionHeirarchy;
+		return SectionHierarchy;
 	}
 
 	/**
@@ -216,13 +222,13 @@ public:
 	template <typename StructType>
 	bool LoadConfig(StructType& Struct, const FString& InterfaceName = FString(), const FString& OperationName = FString()) const
 	{
-		TArray<FString> SectionHeirarchy;
+		TArray<FString> SectionHierarchy;
 
-		auto PushSection = [&SectionHeirarchy](const TArray<const FStringView>& SectionNames)
+		auto PushSection = [&SectionHierarchy](const TArray<const FStringView>& SectionNames)
 		{
 			if (Algo::NoneOf(SectionNames, UE_PROJECTION_MEMBER(FStringView, IsEmpty)))
 			{
-				SectionHeirarchy.Emplace(FString::Join(SectionNames, TEXT(".")));
+				SectionHierarchy.Emplace(FString::Join(SectionNames, TEXT(".")));
 			}
 		};
 
@@ -239,23 +245,29 @@ public:
 		PushSection({ BaseSection, ServiceConfigName, InstanceConfigNameStr, InterfaceName });
 		PushSection({ BaseSection, ServiceConfigName, InstanceConfigNameStr, InterfaceName, OperationName });
 
-		return LoadConfig(Struct, SectionHeirarchy);
+		return LoadConfig(Struct, SectionHierarchy);
 	}
 	
-	UE_DEPRECATED(5.6, "GetConfigSectionHeirachWithOverrides has been renamed GetConfigSectionHeirarchyWithOverrides")
-	TArray<FString> GetConfigSectionHeirachWithOverrides(const TArray<FString>& SectionHeirarchy) const
+	UE_DEPRECATED(5.6, "GetConfigSectionHeirachWithOverrides has been renamed GetConfigSectionHierarchyWithOverrides")
+	TArray<FString> GetConfigSectionHeirachWithOverrides(const TArray<FString>& SectionHierarchy) const
+	{
+		return GetConfigSectionHierarchyWithOverrides(SectionHierarchy);
+	}
+
+	UE_DEPRECATED(5.8, "GetConfigSectionHeirarchyWithOverrides has been renamed GetConfigSectionHierarchyWithOverrides")
+	TArray<FString> GetConfigSectionHeirarchyWithOverrides(const TArray<FString>& SectionHierarchy) const
 	{
-		return GetConfigSectionHeirarchyWithOverrides(SectionHeirarchy);
+		return GetConfigSectionHierarchyWithOverrides(SectionHierarchy);
 	}
 
 	/**
 	 * Get an array of a config section with the overrides added in
 	 * 
-	 * @param SectionHeirarchy Array of config sections to load values from
+	 * @param SectionHierarchy Array of config sections to load values from
 	 * 
 	 * @return Array of the sections with overrides for values to be loaded from
 	 */
-	TArray<FString> GetConfigSectionHeirarchyWithOverrides(const TArray<FString>& SectionHeirarchy) const
+	TArray<FString> GetConfigSectionHierarchyWithOverrides(const TArray<FString>& SectionHierarchy) const
 	{
 		TArray<FString> AllConfigSectionOverrides = ConfigSectionOverrides;
 
@@ -264,37 +276,37 @@ public:
 			AllConfigSectionOverrides.Add("Editor");
 		}
 
-		TArray<FString> SectionHeirarchyWithOverrides;
-		for (const FString& Section : SectionHeirarchy)
+		TArray<FString> SectionHierarchyWithOverrides;
+		for (const FString& Section : SectionHierarchy)
 		{
-			SectionHeirarchyWithOverrides.Add(Section);
+			SectionHierarchyWithOverrides.Add(Section);
 
 			for (const FString& Override : AllConfigSectionOverrides)
 			{
 				// Using of space in online config section is deprecated
 				FString OverrideSection = Section + TEXT(" ") + Override;
-				SectionHeirarchyWithOverrides.Add(OverrideSection);
+				SectionHierarchyWithOverrides.Add(OverrideSection);
 
 				OverrideSection = Section + TEXT(":") + Override;
-				SectionHeirarchyWithOverrides.Add(OverrideSection);
+				SectionHierarchyWithOverrides.Add(OverrideSection);
 			}
 		}
 
-		return SectionHeirarchyWithOverrides;
+		return SectionHierarchyWithOverrides;
 	}
 
 	/**
-	 * Load a config struct for a section heirarchy, also using the ConfigSectionOverrides
+	 * Load a config struct for a section hierarchy, also using the ConfigSectionOverrides
 	 *
 	 * @param Struct Struct to populate with values from config
-	 * @param SectionHeirarchy Array of config sections to load values from
+	 * @param SectionHierarchy Array of config sections to load values from
 	 *
 	 * @return true if a value was loaded
 	 */
 	template <typename StructType>
-	bool LoadConfig(StructType& Struct, const TArray<FString>& SectionHeirarchy) const
+	bool LoadConfig(StructType& Struct, const TArray<FString>& SectionHierarchy) const
 	{
-		return ::UE::Online::LoadConfig(*ConfigProvider, GetConfigSectionHeirarchyWithOverrides(SectionHeirarchy), Struct);
+		return ::UE::Online::LoadConfig(*ConfigProvider, GetConfigSectionHierarchyWithOverrides(SectionHierarchy), Struct);
 	}
 
 	/* Get op (OnlineServices) */
@@ -302,43 +314,43 @@ public:
 	TOnlineAsyncOpRef<OpType> GetOp(typename OpType::Params&& Params)
 	{
 		LogErrorIfOnlineServicesHasShutdown<OpType>();
-		return OpCache.GetOp<OpType>(MoveTemp(Params), GetConfigSectionHeirarchy());
+		return OpCache.GetOp<OpType>(MoveTemp(Params), GetConfigSectionHierarchy());
 	}
 
 	template <typename OpType, typename ParamsFuncsType = TJoinableOpParamsFuncs<OpType>>
 	TOnlineAsyncOpRef<OpType> GetJoinableOp(typename OpType::Params&& Params)
 	{
 		LogErrorIfOnlineServicesHasShutdown<OpType>();
-		return OpCache.GetJoinableOp<OpType, ParamsFuncsType>(MoveTemp(Params), GetConfigSectionHeirarchy());
+		return OpCache.GetJoinableOp<OpType, ParamsFuncsType>(MoveTemp(Params), GetConfigSectionHierarchy());
 	}
 
 	template <typename OpType, typename ParamsFuncsType = TMergeableOpParamsFuncs<OpType>>
 	TOnlineAsyncOpRef<OpType> GetMergeableOp(typename OpType::Params&& Params)
 	{
 		LogErrorIfOnlineServicesHasShutdown<OpType>();
-		return OpCache.GetMergeableOp<OpType, ParamsFuncsType>(MoveTemp(Params), GetConfigSectionHeirarchy());
+		return OpCache.GetMergeableOp<OpType, ParamsFuncsType>(MoveTemp(Params), GetConfigSectionHierarchy());
 	}
 
 	/* Get op (Interface) */
 	template <typename OpType>
-	TOnlineAsyncOpRef<OpType> GetOp(typename OpType::Params&& Params, const TArray<FString>& ConfigSectionHeirarchy)
+	TOnlineAsyncOpRef<OpType> GetOp(typename OpType::Params&& Params, const TArray<FString>& ConfigSectionHierarchy)
 	{
 		LogErrorIfOnlineServicesHasShutdown<OpType>();
-		return OpCache.GetOp<OpType>(MoveTemp(Params), ConfigSectionHeirarchy);
+		return OpCache.GetOp<OpType>(MoveTemp(Params), ConfigSectionHierarchy);
 	}
 
 	template <typename OpType, typename ParamsFuncsType = TJoinableOpParamsFuncs<OpType>>
-	TOnlineAsyncOpRef<OpType> GetJoinableOp(typename OpType::Params&& Params, const TArray<FString>& ConfigSectionHeirarchy)
+	TOnlineAsyncOpRef<OpType> GetJoinableOp(typename OpType::Params&& Params, const TArray<FString>& ConfigSectionHierarchy)
 	{
 		LogErrorIfOnlineServicesHasShutdown<OpType>();
-		return OpCache.GetJoinableOp<OpType, ParamsFuncsType>(MoveTemp(Params), ConfigSectionHeirarchy);
+		return OpCache.GetJoinableOp<OpType, ParamsFuncsType>(MoveTemp(Params), ConfigSectionHierarchy);
 	}
 
 	template <typename OpType, typename ParamsFuncsType = TMergeableOpParamsFuncs<OpType>>
-	TOnlineAsyncOpRef<OpType> GetMergeableOp(typename OpType::Params&& Params, const TArray<FString>& ConfigSectionHeirarchy)
+	TOnlineAsyncOpRef<OpType> GetMergeableOp(typename OpType::Params&& Params, const TArray<FString>& ConfigSectionHierarchy)
 	{
 		LogErrorIfOnlineServicesHasShutdown<OpType>();
-		return OpCache.GetMergeableOp<OpType, ParamsFuncsType>(MoveTemp(Params), ConfigSectionHeirarchy);
+		return OpCache.GetMergeableOp<OpType, ParamsFuncsType>(MoveTemp(Params), ConfigSectionHierarchy);
 	}
 
 	/* Queue for executing tasks in parallel. Serial queues feed into this */
diff --git a/Engine/Plugins/Online/OnlineServices/Source/OnlineServicesInterface/Private/Online/Schema.cpp b/Engine/Plugins/Online/OnlineServices/Source/OnlineServicesInterface/Private/Online/Schema.cpp
--- a/Engine/Plugins/Online/OnlineServices/Source/OnlineServicesInterface/Private/Online/Schema.cpp
+++ b/Engine/Plugins/Online/OnlineServices/Source/OnlineServicesInterface/Private/Online/Schema.cpp
@@ -227,7 +227,7 @@ bool FSchemaRegistry::ParseConfig(const FSchemaRegistryDescriptorConfig& Config)
 		}
 	}
 
-	// Populate known schema. Verify what can be verified statically. Heirarchy validation is performed below.
+	// Populate known schema. Verify what can be verified statically. Hierarchy validation is performed below.
 	TMap<FSchemaId, const FSchemaDescriptor*> KnownSchemaDescriptors;
 	for (const FSchemaDescriptor& SchemaDescriptor : Config.SchemaDescriptors)
 	{
diff --git a/Engine/Source/Programs/Online/OnlineTestsCore/Source/Private/Tests/TestAsyncOp.cpp b/Engine/Source/Programs/Online/OnlineTestsCore/Source/Private/Tests/TestAsyncOp.cpp
--- a/Engine/Source/Programs/Online/OnlineTestsCore/Source/Private/Tests/TestAsyncOp.cpp
+++ b/Engine/Source/Programs/Online/OnlineTestsCore/Source/Private/Tests/TestAsyncOp.cpp
@@ -309,24 +309,24 @@ ASYNCOP_TEST_CASE("Verify that the various OnlineAsyncOp execution policies work
 ASYNCOP_TEST_CASE("Verify that the duration cache policy works)")
 {
 	FTestPipeline& TestPipeline = GetPipeline();
-	TArray<FString> CacheDurationTestHeirarchy = { TEXT("OnlineTests.TestCacheDuration") };
+	TArray<FString> CacheDurationTestHierarchy = { TEXT("OnlineTests.TestCacheDuration") };
 	// Uses the following operation configuration in OnlineTests' DefaultEngine.ini
 	// [OnlineTests.TestCacheDuration]
 	// CacheExpiration=Duration
 	// CacheExpirySeconds=2
 	// bCacheError=false
-	TestPipeline.EmplaceLambda([&CacheDurationTestHeirarchy](const IOnlineServicesPtr& Services)
+	TestPipeline.EmplaceLambda([&CacheDurationTestHierarchy](const IOnlineServicesPtr& Services)
 	{
 		FOnlineServicesCommon& ServicesCommon = static_cast<FOnlineServicesCommon&>(*Services);
-		TOnlineAsyncOpRef<FTestOp> Op = ServicesCommon.GetJoinableOp<FTestOp>(FTestOp::Params(), CacheDurationTestHeirarchy);
+		TOnlineAsyncOpRef<FTestOp> Op = ServicesCommon.GetJoinableOp<FTestOp>(FTestOp::Params(), CacheDurationTestHierarchy);
 		Op->Then([](FOnlineAsyncOp&) {})
 		.Then([](TOnlineAsyncOp<FTestOp>& Op) { Op.SetResult(FTestOp::Result()); })
 		.Enqueue(ServicesCommon.GetParallelQueue());
 	})
-	.EmplaceLambda([&CacheDurationTestHeirarchy](const IOnlineServicesPtr& Services)
+	.EmplaceLambda([&CacheDurationTestHierarchy](const IOnlineServicesPtr& Services)
 	{
 		FOnlineServicesCommon& ServicesCommon = static_cast<FOnlineServicesCommon&>(*Services);
-		TOnlineAsyncOpRef<FTestOp> Op = ServicesCommon.GetJoinableOp<FTestOp>(FTestOp::Params(), CacheDurationTestHeirarchy);
+		TOnlineAsyncOpRef<FTestOp> Op = ServicesCommon.GetJoinableOp<FTestOp>(FTestOp::Params(), CacheDurationTestHierarchy);
 		CHECK(Op->IsReady()); // Cache should return the already in flight/completed op
 	})
 	.EmplaceAsyncLambda([](FAsyncLambdaResult Promise, const IOnlineServicesPtr& Services)
@@ -334,10 +334,10 @@ ASYNCOP_TEST_CASE("Verify that the duration cache policy works)")
 		// Wait a second into the cache duration
 		FTSTicker::GetCoreTicker().AddTicker(TEXT("OnlineTestTimeout"), 1.0f, [Promise](float) { Promise->SetValue(true); return false; });
 	})
-	.EmplaceLambda([&CacheDurationTestHeirarchy](const IOnlineServicesPtr& Services)
+	.EmplaceLambda([&CacheDurationTestHierarchy](const IOnlineServicesPtr& Services)
 	{
 		FOnlineServicesCommon& ServicesCommon = static_cast<FOnlineServicesCommon&>(*Services);
-		TOnlineAsyncOpRef<FTestOp> Op = ServicesCommon.GetJoinableOp<FTestOp>(FTestOp::Params(), CacheDurationTestHeirarchy);
+		TOnlineAsyncOpRef<FTestOp> Op = ServicesCommon.GetJoinableOp<FTestOp>(FTestOp::Params(), CacheDurationTestHierarchy);
 		CHECK(Op->IsReady()); // Cache should return the already completed op
 	})
 	.EmplaceAsyncLambda([](FAsyncLambdaResult Promise, const IOnlineServicesPtr& Services)
@@ -345,10 +345,10 @@ ASYNCOP_TEST_CASE("Verify that the duration cache policy works)")
 		// Wait until shortly after the cache should expire
 		FTSTicker::GetCoreTicker().AddTicker(TEXT("OnlineTestTimeout"), 2.0f, [Promise](float) { Promise->SetValue(true); return false; });
 	})
-	.EmplaceAsyncLambda([&CacheDurationTestHeirarchy](FAsyncLambdaResult Promise, const IOnlineServicesPtr& Services)
+	.EmplaceAsyncLambda([&CacheDurationTestHierarchy](FAsyncLambdaResult Promise, const IOnlineServicesPtr& Services)
 	{
 		FOnlineServicesCommon& ServicesCommon = static_cast<FOnlineServicesCommon&>(*Services);
-		TOnlineAsyncOpRef<FTestOp> Op = ServicesCommon.GetJoinableOp<FTestOp>(FTestOp::Params(), CacheDurationTestHeirarchy);
+		TOnlineAsyncOpRef<FTestOp> Op = ServicesCommon.GetJoinableOp<FTestOp>(FTestOp::Params(), CacheDurationTestHierarchy);
 		CHECK(!Op->IsReady()); // Cache should return a new op
 		Op->Then([](FOnlineAsyncOp&) {})
 		.Then([Promise](TOnlineAsyncOp<FTestOp>& Op) { Op.SetResult(FTestOp::Result()); Promise->SetValue(true); })
@@ -366,39 +366,39 @@ ASYNCOP_TEST_CASE("Verify that the duration cache policy works)")
 ASYNCOP_TEST_CASE("Verify that the cache error setting works)", ASYNCOP_SKIP_EOS_PLATFORM_FAULT_TAG)
 {
 	FTestPipeline& TestPipeline = GetPipeline();
-	TArray<FString> DontCacheErrorHeirarchy = { TEXT("OnlineTests.TestDontCacheError") };
+	TArray<FString> DontCacheErrorHierarchy = { TEXT("OnlineTests.TestDontCacheError") };
 	// Uses the following operation configuration in OnlineTests' DefaultEngine.ini
 	// [OnlineTests.TestDontCacheError]
 	// CacheExpiration=Duration
 	// CacheExpirySeconds=1
 	// bCacheError=false
-	TArray<FString> CacheErrorHeirarchy = { TEXT("OnlineTests.TestCacheError") };
+	TArray<FString> CacheErrorHierarchy = { TEXT("OnlineTests.TestCacheError") };
 	// Uses the following operation configuration in OnlineTests' DefaultEngine.ini
 	// [OnlineTests.TestCacheError]
 	// CacheExpiration=Duration
 	// CacheExpirySeconds=1
 	// bCacheError=true
-	TestPipeline.EmplaceAsyncLambda([&DontCacheErrorHeirarchy](FAsyncLambdaResult Promise, const IOnlineServicesPtr& Services)
+	TestPipeline.EmplaceAsyncLambda([&DontCacheErrorHierarchy](FAsyncLambdaResult Promise, const IOnlineServicesPtr& Services)
 	{
 		FOnlineServicesCommon& ServicesCommon = static_cast<FOnlineServicesCommon&>(*Services);
-		TOnlineAsyncOpRef<FTestOp> Op = ServicesCommon.GetJoinableOp<FTestOp>(FTestOp::Params(), DontCacheErrorHeirarchy);
+		TOnlineAsyncOpRef<FTestOp> Op = ServicesCommon.GetJoinableOp<FTestOp>(FTestOp::Params(), DontCacheErrorHierarchy);
 		Op->Then([](FOnlineAsyncOp&) {})
 		.Then([Promise](TOnlineAsyncOp<FTestOp>& Op) { Op.SetError(Errors::Unknown()); Promise->SetValue(true); })
 		.Enqueue(ServicesCommon.GetParallelQueue());
 	})
-	.EmplaceAsyncLambda([&CacheErrorHeirarchy](FAsyncLambdaResult Promise, const IOnlineServicesPtr& Services)
+	.EmplaceAsyncLambda([&CacheErrorHierarchy](FAsyncLambdaResult Promise, const IOnlineServicesPtr& Services)
 	{
 		FOnlineServicesCommon& ServicesCommon = static_cast<FOnlineServicesCommon&>(*Services);
-		TOnlineAsyncOpRef<FTestOp> Op = ServicesCommon.GetJoinableOp<FTestOp>(FTestOp::Params(), CacheErrorHeirarchy);
+		TOnlineAsyncOpRef<FTestOp> Op = ServicesCommon.GetJoinableOp<FTestOp>(FTestOp::Params(), CacheErrorHierarchy);
 		CHECK(!Op->IsReady()); // Should return a new op
 		Op->Then([](FOnlineAsyncOp&) {})
 		.Then([Promise](TOnlineAsyncOp<FTestOp>& Op) { Op.SetError(Errors::Unknown()); Promise->SetValue(true); })
 		.Enqueue(ServicesCommon.GetParallelQueue());
 	})
-	.EmplaceAsyncLambda([&CacheErrorHeirarchy](FAsyncLambdaResult Promise, const IOnlineServicesPtr& Services)
+	.EmplaceAsyncLambda([&CacheErrorHierarchy](FAsyncLambdaResult Promise, const IOnlineServicesPtr& Services)
 	{
 		FOnlineServicesCommon& ServicesCommon = static_cast<FOnlineServicesCommon&>(*Services);
-		TOnlineAsyncOpRef<FTestOp> Op = ServicesCommon.GetJoinableOp<FTestOp>(FTestOp::Params(), CacheErrorHeirarchy);
+		TOnlineAsyncOpRef<FTestOp> Op = ServicesCommon.GetJoinableOp<FTestOp>(FTestOp::Params(), CacheErrorHierarchy);
 		CHECK(Op->IsReady()); // Should return the cached, failed op
 		Op->GetHandle().OnComplete([Promise](const TOnlineResult<FTestOp>& Result)
 		{
@@ -418,26 +418,26 @@ ASYNCOP_TEST_CASE("Verify that the cache error setting works)", ASYNCOP_SKIP_EOS
 ASYNCOP_TEST_CASE("Verify that the cache expires once complete)")
 {
 	FTestPipeline& TestPipeline = GetPipeline();
-	TArray<FString> CacheExpireUponCompletionHeirarchy = { TEXT("OnlineTests.TestCacheExpireUponCompletion") };
+	TArray<FString> CacheExpireUponCompletionHierarchy = { TEXT("OnlineTests.TestCacheExpireUponCompletion") };
 	// Uses the following operation configuration in OnlineTests' DefaultEngine.ini
 	// [OnlineTests.TestCacheExpireUponCompletion]
 	// CacheExpiration=UponCompletion
 	// bCacheError=false
-	TestPipeline.EmplaceAsyncLambda([&CacheExpireUponCompletionHeirarchy](FAsyncLambdaResult Promise, const IOnlineServicesPtr& Services)
+	TestPipeline.EmplaceAsyncLambda([&CacheExpireUponCompletionHierarchy](FAsyncLambdaResult Promise, const IOnlineServicesPtr& Services)
 	{
 		FOnlineServicesCommon& ServicesCommon = static_cast<FOnlineServicesCommon&>(*Services);
-		TOnlineAsyncOpRef<FTestOp> Op = ServicesCommon.GetJoinableOp<FTestOp>(FTestOp::Params(), CacheExpireUponCompletionHeirarchy);
+		TOnlineAsyncOpRef<FTestOp> Op = ServicesCommon.GetJoinableOp<FTestOp>(FTestOp::Params(), CacheExpireUponCompletionHierarchy);
 		Op->Then([](FOnlineAsyncOp&) {})
 		.Then([Promise](TOnlineAsyncOp<FTestOp>& Op) { Op.SetResult(FTestOp::Result()); Promise->SetValue(true); })
 		.Enqueue(ServicesCommon.GetParallelQueue());
-		TOnlineAsyncOpRef<FTestOp> Op2 = ServicesCommon.GetJoinableOp<FTestOp>(FTestOp::Params(), CacheExpireUponCompletionHeirarchy);
+		TOnlineAsyncOpRef<FTestOp> Op2 = ServicesCommon.GetJoinableOp<FTestOp>(FTestOp::Params(), CacheExpireUponCompletionHierarchy);
 		CHECK(Op2->IsReady());
 		CHECK(!Op2->IsComplete());
 	})
-	.EmplaceAsyncLambda([&CacheExpireUponCompletionHeirarchy](FAsyncLambdaResult Promise, const IOnlineServicesPtr& Services)
+	.EmplaceAsyncLambda([&CacheExpireUponCompletionHierarchy](FAsyncLambdaResult Promise, const IOnlineServicesPtr& Services)
 	{
 		FOnlineServicesCommon& ServicesCommon = static_cast<FOnlineServicesCommon&>(*Services);
-		TOnlineAsyncOpRef<FTestOp> Op = ServicesCommon.GetJoinableOp<FTestOp>(FTestOp::Params(), CacheExpireUponCompletionHeirarchy);
+		TOnlineAsyncOpRef<FTestOp> Op = ServicesCommon.GetJoinableOp<FTestOp>(FTestOp::Params(), CacheExpireUponCompletionHierarchy);
 		CHECK(!Op->IsReady());
 		Op->Then([](FOnlineAsyncOp&) {})
 		.Then([Promise](TOnlineAsyncOp<FTestOp>& Op) { Op.SetResult(FTestOp::Result()); Promise->SetValue(true); })
