UE: Fix memory leak in FOnlineAsyncTaskManager cleanup
- Moving Clean up of all task queues to the base class destructor from FOnlineAsyncTaskManagerGDK destructor
- Adding a test to check FOnlineAsyncTaskManager destructor cleanup

#tests onlinetests

#Synced-CL 52095155

diff --git a/Engine/Plugins/Online/Microsoft/OnlineSubsystemGDK/Source/Private/OnlineAsyncTaskManagerGDK.cpp b/Engine/Plugins/Online/Microsoft/OnlineSubsystemGDK/Source/Private/OnlineAsyncTaskManagerGDK.cpp
--- a/Engine/Plugins/Online/Microsoft/OnlineSubsystemGDK/Source/Private/OnlineAsyncTaskManagerGDK.cpp
+++ b/Engine/Plugins/Online/Microsoft/OnlineSubsystemGDK/Source/Private/OnlineAsyncTaskManagerGDK.cpp
@@ -9,44 +9,8 @@
 
 FOnlineAsyncTaskManagerGDK::~FOnlineAsyncTaskManagerGDK()
 {
-	// cancel any in-flight tasks
+	// Cancel any GDK-specific in-flight tasks before base class destructor cleans up the queues
 	OnlineTasksQueue.CancelPendingTasksAndDestroyQueue();
-
-	// delete all cached tasks
-	{
-		FScopeLock Lock(&QueuedSerialTasksLock);
-		for (FOnlineAsyncTask* Task : QueuedSerialTasks)
-		{
-			delete Task;
-		}
-		QueuedSerialTasks.Reset();
-	}
-
-	{
-		FOnlineAsyncTask* ParallelTask = nullptr;
-		while (QueuedParallelTasks.Dequeue(ParallelTask))
-		{
-			delete ParallelTask;
-		}
-	}
-
-	{
-		FScopeLock Lock(&ParallelTasksLock);
-		for (FOnlineAsyncTask* Task : ParallelTasks)
-		{
-			delete Task;
-		}
-		ParallelTasks.Reset();
-	}
-
-	{
-		FScopeLock Lock(&OutQueueLock);
-		for (FOnlineAsyncItem* Task : OutQueue)
-		{
-			delete Task;
-		}
-		OutQueue.Reset();
-	}
 }
 
 void FOnlineAsyncTaskManagerGDK::OnlineTick()
diff --git a/Engine/Plugins/Online/OnlineSubsystem/Source/Private/OnlineAsyncTaskManager.cpp b/Engine/Plugins/Online/OnlineSubsystem/Source/Private/OnlineAsyncTaskManager.cpp
--- a/Engine/Plugins/Online/OnlineSubsystem/Source/Private/OnlineAsyncTaskManager.cpp
+++ b/Engine/Plugins/Online/OnlineSubsystem/Source/Private/OnlineAsyncTaskManager.cpp
@@ -48,6 +48,40 @@ FOnlineAsyncTaskManager::FOnlineAsyncTaskManager() :
 
 FOnlineAsyncTaskManager::~FOnlineAsyncTaskManager()
 {
+	// Clean up the active serial task
+	delete ActiveSerialTask;
+	ActiveSerialTask = nullptr;
+
+	// Clean up queued serial tasks
+	for (FOnlineAsyncTask* Task : QueuedSerialTasks)
+	{
+		delete Task;
+	}
+	QueuedSerialTasks.Reset();
+
+	// Clean up queued parallel tasks (TQueue)
+	{
+		FOnlineAsyncTask* Task = nullptr;
+		while (QueuedParallelTasks.Dequeue(Task))
+		{
+			delete Task;
+		}
+	}
+
+	// Clean up in-progress parallel tasks
+	for (FOnlineAsyncTask* Task : ParallelTasks)
+	{
+		delete Task;
+	}
+	ParallelTasks.Reset();
+
+	// Clean up completed tasks in the out queue
+	for (FOnlineAsyncItem* Item : OutQueue)
+	{
+		delete Item;
+	}
+	OutQueue.Reset();
+
 	FPlatformProcess::ReturnSynchEventToPool(WorkEvent);
 }
 
diff --git a/Engine/Source/Programs/Online/OnlineTestsCore/Source/Private/Tests/AsyncTaskManager/AsyncTaskManagerCleanupTests.cpp b/Engine/Source/Programs/Online/OnlineTestsCore/Source/Private/Tests/AsyncTaskManager/AsyncTaskManagerCleanupTests.cpp
new file mode 100644
--- /dev/null
+++ b/Engine/Source/Programs/Online/OnlineTestsCore/Source/Private/Tests/AsyncTaskManager/AsyncTaskManagerCleanupTests.cpp
@@ -0,0 +1,121 @@
+// Copyright Epic Games, Inc. All Rights Reserved.
+
+#include "OnlineAsyncTaskManager.h"
+
+#include <catch2/catch_test_macros.hpp>
+
+#define ASYNCTASKMANAGER_TAG "[AsyncTaskManager]"
+#define ASYNCTASKMANAGER_TEST_CASE(x, ...) TEST_CASE(x, ASYNCTASKMANAGER_TAG __VA_ARGS__)
+
+namespace
+{
+
+/** Tracks how many instances have been destroyed via a shared counter */
+class FMockOnlineAsyncTask : public FOnlineAsyncTask
+{
+public:
+	FMockOnlineAsyncTask(int32& InDestroyCounter)
+		: DestroyCounter(InDestroyCounter)
+	{
+	}
+
+	virtual ~FMockOnlineAsyncTask() override
+	{
+		++DestroyCounter;
+	}
+
+	virtual FString ToString() const override { return TEXT("MockAsyncTask"); }
+	virtual bool IsDone() const override { return false; }
+	virtual bool WasSuccessful() const override { return false; }
+
+private:
+	int32& DestroyCounter;
+};
+
+/** FOnlineAsyncItem subclass for testing OutQueue cleanup */
+class FMockOnlineAsyncItem : public FOnlineAsyncItem
+{
+public:
+	FMockOnlineAsyncItem(int32& InDestroyCounter)
+		: DestroyCounter(InDestroyCounter)
+	{
+	}
+
+	virtual ~FMockOnlineAsyncItem() override
+	{
+		++DestroyCounter;
+	}
+
+	virtual FString ToString() const override { return TEXT("MockAsyncItem"); }
+
+private:
+	int32& DestroyCounter;
+};
+
+/** Minimal concrete FOnlineAsyncTaskManager for testing */
+class FMockOnlineAsyncTaskManager : public FOnlineAsyncTaskManager
+{
+public:
+	virtual void OnlineTick() override {}
+};
+
+} // anonymous namespace
+
+ASYNCTASKMANAGER_TEST_CASE("Destructor cleans up queued serial tasks")
+{
+	int32 DestroyCount = 0;
+	{
+		FMockOnlineAsyncTaskManager Manager;
+		Manager.AddToInQueue(new FMockOnlineAsyncTask(DestroyCount));
+		Manager.AddToInQueue(new FMockOnlineAsyncTask(DestroyCount));
+		Manager.AddToInQueue(new FMockOnlineAsyncTask(DestroyCount));
+	}
+	CHECK(DestroyCount == 3);
+}
+
+ASYNCTASKMANAGER_TEST_CASE("Destructor cleans up queued parallel tasks")
+{
+	int32 DestroyCount = 0;
+	{
+		FMockOnlineAsyncTaskManager Manager;
+		Manager.AddToParallelTasks(new FMockOnlineAsyncTask(DestroyCount));
+		Manager.AddToParallelTasks(new FMockOnlineAsyncTask(DestroyCount));
+	}
+	CHECK(DestroyCount == 2);
+}
+
+ASYNCTASKMANAGER_TEST_CASE("Destructor cleans up out queue")
+{
+	int32 DestroyCount = 0;
+	{
+		FMockOnlineAsyncTaskManager Manager;
+		Manager.AddToOutQueue(new FMockOnlineAsyncItem(DestroyCount));
+		Manager.AddToOutQueue(new FMockOnlineAsyncItem(DestroyCount));
+	}
+	CHECK(DestroyCount == 2);
+}
+
+ASYNCTASKMANAGER_TEST_CASE("Destructor cleans up all queues simultaneously")
+{
+	int32 SerialDestroyCount = 0;
+	int32 ParallelDestroyCount = 0;
+	int32 OutQueueDestroyCount = 0;
+	{
+		FMockOnlineAsyncTaskManager Manager;
+		Manager.AddToInQueue(new FMockOnlineAsyncTask(SerialDestroyCount));
+		Manager.AddToInQueue(new FMockOnlineAsyncTask(SerialDestroyCount));
+		Manager.AddToParallelTasks(new FMockOnlineAsyncTask(ParallelDestroyCount));
+		Manager.AddToOutQueue(new FMockOnlineAsyncItem(OutQueueDestroyCount));
+		Manager.AddToOutQueue(new FMockOnlineAsyncItem(OutQueueDestroyCount));
+		Manager.AddToOutQueue(new FMockOnlineAsyncItem(OutQueueDestroyCount));
+	}
+	CHECK(SerialDestroyCount == 2);
+	CHECK(ParallelDestroyCount == 1);
+	CHECK(OutQueueDestroyCount == 3);
+}
+
+ASYNCTASKMANAGER_TEST_CASE("Destructor handles empty queues gracefully")
+{
+	// Should not crash when destroying a manager with no tasks
+	FMockOnlineAsyncTaskManager Manager;
+}
